Add a project setting to not register the default raw input devices if you're already handling the XInput devices through another manner

diff --git a/Engine/Plugins/Experimental/RawInput/Source/RawInput/Private/Windows/RawInputWindows.cpp b/Engine/Plugins/Experimental/RawInput/Source/RawInput/Private/Windows/RawInputWindows.cpp
--- a/Engine/Plugins/Experimental/RawInput/Source/RawInput/Private/Windows/RawInputWindows.cpp
+++ b/Engine/Plugins/Experimental/RawInput/Source/RawInput/Private/Windows/RawInputWindows.cpp
@@ -94,16 +94,19 @@ FRawInputWindows::FRawInputWindows(const TSharedRef<FGenericApplicationMessageHa
 	WindowsApplication->AddMessageHandler(*this);
 	QueryConnectedDevices();
 
-	// Register a default device.
-	const uint32 Flags = 0;
-	const int32 PageID = 0x01;
-	int32 DeviceID = 0x04;
-	DefaultDeviceHandle = RegisterInputDevice(RIM_TYPEHID, Flags, DeviceID, PageID);
-
-	if(DefaultDeviceHandle==INDEX_NONE)
+	// Register a default device if desired
+	if (GetDefault<URawInputSettings>()->bRegisterDefaultDevice)
 	{
-		DeviceID = 0x05;
+		const uint32 Flags = 0;
+		const int32 PageID = 0x01;
+		int32 DeviceID = 0x04;
 		DefaultDeviceHandle = RegisterInputDevice(RIM_TYPEHID, Flags, DeviceID, PageID);
+
+		if (DefaultDeviceHandle == INDEX_NONE)
+		{
+			DeviceID = 0x05;
+			DefaultDeviceHandle = RegisterInputDevice(RIM_TYPEHID, Flags, DeviceID, PageID);
+		}
 	}
 
 	AHUD::OnShowDebugInfo.AddRaw(this, &FRawInputWindows::ShowDebugInfo);
diff --git a/Engine/Plugins/Experimental/RawInput/Source/RawInput/Public/RawInputSettings.h b/Engine/Plugins/Experimental/RawInput/Source/RawInput/Public/RawInputSettings.h
--- a/Engine/Plugins/Experimental/RawInput/Source/RawInput/Public/RawInputSettings.h
+++ b/Engine/Plugins/Experimental/RawInput/Source/RawInput/Public/RawInputSettings.h
@@ -100,5 +100,8 @@ public:
 
 	UPROPERTY(config, EditAnywhere, Category="Device Configurations")
 	TArray<FRawInputDeviceConfiguration> DeviceConfigurations;
+
+	UPROPERTY(config, EditAnywhere, Category="Device Configurations")
+	bool bRegisterDefaultDevice = true;
 };
 
