PR #6093 : RawInput plugin fixes to support generic gamepads on windows (Contributed by [REDACTED_PERSON])

diff --git a/Engine/Plugins/Experimental/RawInput/Source/RawInput/Private/Windows/RawInputWindows.cpp b/Engine/Plugins/Experimental/RawInput/Source/RawInput/Private/Windows/RawInputWindows.cpp
--- a/Engine/Plugins/Experimental/RawInput/Source/RawInput/Private/Windows/RawInputWindows.cpp
+++ b/Engine/Plugins/Experimental/RawInput/Source/RawInput/Private/Windows/RawInputWindows.cpp
@@ -179,13 +179,7 @@ int32 FRawInputWindows::RegisterInputDevice(const int32 DeviceType, const int32
 				if (CompareDeviceInfo(ConnectedDeviceInfo.RIDDeviceInfo, DeviceData))
 				{
 					FRawWindowsDeviceEntry& RegisteredDeviceInfo = RegisteredDeviceList[DeviceHandle];
-					RegisteredDeviceInfo.bIsConnected = true;
-					RegisteredDeviceInfo.DeviceData.DeviceName = ConnectedDeviceInfo.DeviceName;
-					if (DeviceData.DeviceType == RIM_TYPEHID)
-					{
-						RegisteredDeviceInfo.DeviceData.VendorID = ConnectedDeviceInfo.RIDDeviceInfo.hid.dwVendorId;
-						RegisteredDeviceInfo.DeviceData.ProductID = ConnectedDeviceInfo.RIDDeviceInfo.hid.dwProductId;
-					}
+					CopyConnectedDeviceInfo(RegisteredDeviceInfo, &ConnectedDeviceInfo);
 
 					UE_LOG(LogRawInputWindows, Log, TEXT("VenderID:%x ProductID:%x"), RegisteredDeviceInfo.DeviceData.VendorID, RegisteredDeviceInfo.DeviceData.ProductID);
 
@@ -251,7 +245,8 @@ void FRawInputWindows::SetupBindings(const int32 DeviceHandle, const bool bApply
 		const int32 VendorID = FCString::Strtoi(*DeviceConfig.VendorID, nullptr, 16);
 		const int32 ProductID = FCString::Strtoi(*DeviceConfig.ProductID, nullptr, 16);
 
-		if (VendorID != 0 && ProductID != 0 && VendorID == DeviceEntry.DeviceData.VendorID && ProductID == DeviceEntry.DeviceData.ProductID)
+		// If VendorId or ProductId are 0, apply to everything
+		if ((VendorID == 0 || VendorID == DeviceEntry.DeviceData.VendorID) && (ProductID == 0 || ProductID == DeviceEntry.DeviceData.ProductID))
 		{
 			const int32 NumButtons = FMath::Min(DeviceConfig.ButtonProperties.Num(), MAX_NUM_CONTROLLER_BUTTONS);
 			DeviceEntry.ButtonData.SetNum(NumButtons);
@@ -272,6 +267,7 @@ void FRawInputWindows::SetupBindings(const int32 DeviceHandle, const bool bApply
 					AnalogData.KeyName = AxisProps.Key.GetFName();
 					AnalogData.Offset = AxisProps.Offset;
 					AnalogData.bInverted = AxisProps.bInverted;
+					AnalogData.bGamepadStick = AxisProps.bGamepadStick;
 				}
 				else
 				{
@@ -447,47 +443,55 @@ bool FRawInputWindows::ProcessMessage(const HWND hwnd, const uint32 Msg, const W
 
  			bool bIsRegisteredDevice = false;			
 			
-			for (const TPair<int32, FRawWindowsDeviceEntry>& DeviceEntryPair : RegisteredDeviceList)
+			// If this is a HID device we need to do some stuff to determine if its one we care about (IE one we registered)
+			if (RawInputDataBuffer->header.dwType == RIM_TYPEHID)
 			{
-				const FRawWindowsDeviceEntry& EachEntry = DeviceEntryPair.Value;
-				if (RawInputDataBuffer->header.dwType == EachEntry.DeviceData.DeviceType)
+				HIDP_CAPS Caps;
+				FRawInputRegisteredDevice DeviceData;
+
+				// First we need to get the pre-parsed data
+				uint32 BufferSize;
+
+				if (::GetRawInputDeviceInfo(RawInputDataBuffer->header.hDevice, RIDI_PREPARSEDDATA, nullptr, &BufferSize) != RAW_INPUT_ERROR)
 				{
-					// If this is a HID device we need to do some stuff to determine if its one we care about (IE one we registered)
-					if (RawInputDataBuffer->header.dwType == RIM_TYPEHID)
+					PreParsedData.SetNumUninitialized(BufferSize + 1, false);
+							
+					if (::GetRawInputDeviceInfo(RawInputDataBuffer->header.hDevice, RIDI_PREPARSEDDATA, PreParsedData.GetData(), &BufferSize) != RAW_INPUT_ERROR)
 					{
-						// First we need to get the pre-parsed data
-						uint32 BufferSize;
+						// now that we have the PP data we need to get the caps, check those and see if this is a device we registered and if it is store it so we can send it
+						if (DLLPointers.HidP_GetCaps((PHIDP_PREPARSED_DATA)PreParsedData.GetData(), &Caps) == HIDP_STATUS_SUCCESS)
+						{
+							DeviceData = FRawInputRegisteredDevice(RawInputDataBuffer->header.dwType, Caps.Usage, Caps.UsagePage);									
+						}
+					}							
+				}
 
-						if (::GetRawInputDeviceInfo(RawInputDataBuffer->header.hDevice, RIDI_PREPARSEDDATA, nullptr, &BufferSize) != RAW_INPUT_ERROR)
+				if (DeviceData.bIsValid)
+				{
+					// Search for a registered device matching details
+					for (const TPair<int32, FRawWindowsDeviceEntry>& DeviceEntryPair : RegisteredDeviceList)
+					{
+						const FRawWindowsDeviceEntry& EachEntry = DeviceEntryPair.Value;
+						if (DeviceData == EachEntry.DeviceData)
 						{
-							PreParsedData.SetNumUninitialized(BufferSize + 1, false);
-							
-							if (::GetRawInputDeviceInfo(RawInputDataBuffer->header.hDevice, RIDI_PREPARSEDDATA, PreParsedData.GetData(), &BufferSize) != RAW_INPUT_ERROR)
+							if (!EachEntry.bIsConnected)
 							{
-								HIDP_CAPS Caps;
-
-								// now that we have the PP data we need to get the caps, check those and see if this is a device we registered and if it is store it so we can send it
-								if (DLLPointers.HidP_GetCaps((PHIDP_PREPARSED_DATA)PreParsedData.GetData(), &Caps) == HIDP_STATUS_SUCCESS)
-								{
-									FRawInputRegisteredDevice DeviceData(RawInputDataBuffer->header.dwType, Caps.Usage, Caps.UsagePage);
-									// Win32 doesn't correctly report the device ID, so at least for now just trust it is from the device we want
-									//if (DeviceData == EachEntry.DeviceData)
-									{
-										bIsRegisteredDevice = true;
-										ParseInputData(DeviceEntryPair.Key, RawInputDataBuffer, (PHIDP_PREPARSED_DATA)PreParsedData.GetData(), Caps);
-									}									
-								}
-							}							
+								// Repoll connection data as this claims to be disconnected
+								QueryConnectedDevices();
+							}
+
+							bIsRegisteredDevice = true;
+							ParseInputData(DeviceEntryPair.Key, RawInputDataBuffer, (PHIDP_PREPARSED_DATA)PreParsedData.GetData(), Caps);
 						}
 					}
-					else
-					{
-						// Must be a keyboard/mouse, just send the data as we don't really have any detailed info about those to check if we registered them
-						bIsRegisteredDevice = true;
-					}
-
 				}
 			}
+			else
+			{
+				// Must be a keyboard/mouse, just send the data as we don't really have any detailed info about those to check if we registered them
+				bIsRegisteredDevice = true;
+			}
+			
 			if (bIsRegisteredDevice && FilteredInputDataHandler.IsBound())
 			{			
 				bHandled = FilteredInputDataHandler.Execute(Size, RawInputDataBuffer);
@@ -583,7 +587,7 @@ void FRawInputWindows::ParseInputData(const int32 InHandle, const RAWINPUT* InRa
 		{
 			for (FButtonData& DeviceButtonData : DeviceEntry->ButtonData)
 			{
-				DeviceButtonData.bPreviousButtonState = DeviceButtonData.bButtonState;
+				// We don't reset previous button state so we can detect state transitions, there's no guarantee this will get called after input has consumed the previous state
 				DeviceButtonData.bButtonState = false;
 			}
 
@@ -637,6 +641,12 @@ void FRawInputWindows::ParseInputData(const int32 InHandle, const RAWINPUT* InRa
 						if (DeviceEntry->AnalogData[iValue].RangeMax == -1.f)
 						{
 							DeviceEntry->AnalogData[iValue].RangeMax = ValueCapsBuffer[iValue].LogicalMax;
+							if (DeviceEntry->AnalogData[iValue].RangeMax < DeviceEntry->AnalogData[iValue].RangeMin)
+							{
+								// Need to mask against BitSize, xinput devices return 0xffffffff when they really mean 0x0000ffff
+								LONG BitMask = (1 << ValueCapsBuffer[iValue].BitSize) - 1;
+								DeviceEntry->AnalogData[iValue].RangeMax = (ValueCapsBuffer[iValue].LogicalMax & BitMask);
+							}
 						}
 						DeviceEntry->AnalogData[iValue].Index = ValueCapsBuffer[iValue].Range.UsageMin;
 						DeviceEntry->AnalogData[iValue].Value = (float)EachValue;
@@ -713,7 +723,38 @@ void FRawInputWindows::QueryConnectedDevices()
 		ShowDeviceInfo(ConnectedDeviceInfoList.Last());
 	}
 
-	UE_LOG(LogRawInputWindows, Warning, TEXT("Found device %d devices"), RegisteredDeviceList.Num());
+	// Update registered device info connected status
+	for (TPair<int32, FRawWindowsDeviceEntry>& DeviceEntryPair : RegisteredDeviceList)
+	{
+		bool bFoundConnected = false;
+		for (const FConnectedDeviceInfo& ConnectedDeviceInfo : ConnectedDeviceInfoList)
+		{
+			if (CompareDeviceInfo(ConnectedDeviceInfo.RIDDeviceInfo, DeviceEntryPair.Value.DeviceData))
+			{
+				bool bWasConnected = DeviceEntryPair.Value.bIsConnected;
+
+				bFoundConnected = true;
+				CopyConnectedDeviceInfo(DeviceEntryPair.Value, &ConnectedDeviceInfo);
+
+				if (!bWasConnected)
+				{
+					// If this is the first connection, apply bindings now
+					SetupBindings(DeviceEntryPair.Key, true);
+					
+					UE_LOG(LogRawInputWindows, Log, TEXT("Device was connected after registration (Usage:%d UsagePage:%d)"), DeviceEntryPair.Value.DeviceData.Usage, DeviceEntryPair.Value.DeviceData.UsagePage);				
+				}
+
+				break;
+			}
+		}
+
+		if (!bFoundConnected)
+		{
+			CopyConnectedDeviceInfo(DeviceEntryPair.Value, nullptr);
+		}
+	}
+
+	UE_LOG(LogRawInputWindows, Warning, TEXT("Found device %d devices"), ConnectedDeviceInfoList.Num());
 }
 
 FString FRawInputWindows::GetErrorString(const int32 StatusCode) const
@@ -803,6 +844,24 @@ bool FRawInputWindows::CompareDeviceInfo(const RID_DEVICE_INFO& DeviceInfo, cons
 	return bResult;
 }
 
+void FRawInputWindows::CopyConnectedDeviceInfo(FRawWindowsDeviceEntry& RegisteredDevice, const FConnectedDeviceInfo* ConnectedDevice)
+{
+	if (ConnectedDevice)
+	{
+		RegisteredDevice.bIsConnected = true;
+		RegisteredDevice.DeviceData.DeviceName = ConnectedDevice->DeviceName;
+		if (RegisteredDevice.DeviceData.DeviceType == RIM_TYPEHID)
+		{
+			RegisteredDevice.DeviceData.VendorID = ConnectedDevice->RIDDeviceInfo.hid.dwVendorId;
+			RegisteredDevice.DeviceData.ProductID = ConnectedDevice->RIDDeviceInfo.hid.dwProductId;
+		}
+	}
+	else
+	{
+		RegisteredDevice.bIsConnected = false;
+	}
+}
+
 void FRawInputWindows::ShowDeviceInfo(const FConnectedDeviceInfo& DeviceInfo) const
 {
 	UE_LOG(LogRawInputWindows, Verbose, TEXT("%s"), *DeviceInfo.DeviceName);
diff --git a/Engine/Plugins/Experimental/RawInput/Source/RawInput/Public/RawInputSettings.h b/Engine/Plugins/Experimental/RawInput/Source/RawInput/Public/RawInputSettings.h
--- a/Engine/Plugins/Experimental/RawInput/Source/RawInput/Public/RawInputSettings.h
+++ b/Engine/Plugins/Experimental/RawInput/Source/RawInput/Public/RawInputSettings.h
@@ -16,6 +16,7 @@ struct RAWINPUT_API FRawInputDeviceAxisProperties
 	FRawInputDeviceAxisProperties()
 		: bEnabled(true)
 		, bInverted(false)
+		, bGamepadStick(false)
 		, Offset(0.f)
 	{
 	}
@@ -32,6 +33,10 @@ struct RAWINPUT_API FRawInputDeviceAxisProperties
 	UPROPERTY(EditAnywhere, Category="Axis Properties", meta=(EditCondition="bEnabled"))
 	uint8 bInverted:1;
 
+	/** Whether this is a gamepad stick, which means it should scale from -1 to 1 instead of 0 to 1 */
+	UPROPERTY(EditAnywhere, Category="Axis Properties", meta=(EditCondition="bEnabled"))
+	uint8 bGamepadStick:1;
+
 	/** The amount to offset the axis value by to get the desired range (note that this offset is applied after the value has been inverted). */
 	UPROPERTY(EditAnywhere, Category="Axis Properties", meta=(EditCondition="bEnabled"))
 	float Offset;
@@ -63,11 +68,11 @@ struct RAWINPUT_API FRawInputDeviceConfiguration
 
 	FRawInputDeviceConfiguration();
 
-	// The hexadecimal vendor ID in string form (i.e. 0x046D)
+	/** The hexadecimal vendor ID in string form (e.g. 0x046D). If this is empty, will match all vendors */
 	UPROPERTY(EditAnywhere, Category="Device Config")
 	FString VendorID;
 
-	// The hexadecimal product ID in string form (i.e. 0xC262)
+	/** The hexadecimal product ID in string form (e.g. 0xC262). If this is empty, will match all products */
 	UPROPERTY(EditAnywhere, Category="Device Config")
 	FString ProductID;
 
@@ -98,9 +103,11 @@ public:
 #endif
 	// End UDeveloperSettings Interface
 
+	/** List of bindings to apply based on Vendor and Product Id, these are applied in order. To set a default, add an entry with empty product/vendor id at the end of the array */
 	UPROPERTY(config, EditAnywhere, Category="Device Configurations")
 	TArray<FRawInputDeviceConfiguration> DeviceConfigurations;
 
+	/** Whether the plugin should try to register a default device to handle generic gamepads and joysticks */
 	UPROPERTY(config, EditAnywhere, Category="Device Configurations")
 	bool bRegisterDefaultDevice = true;
 };
diff --git a/Engine/Plugins/Experimental/RawInput/Source/RawInput/Public/Windows/RawInputWindows.h b/Engine/Plugins/Experimental/RawInput/Source/RawInput/Public/Windows/RawInputWindows.h
--- a/Engine/Plugins/Experimental/RawInput/Source/RawInput/Public/Windows/RawInputWindows.h
+++ b/Engine/Plugins/Experimental/RawInput/Source/RawInput/Public/Windows/RawInputWindows.h
@@ -26,14 +26,14 @@ struct FDLLPointers
 public:
 	
 	/* typedefs for the functions we want */
-	typedef bool(*HidD_GetSerialNumberStringType)(HANDLE device, PVOID buffer, uint32 buffer_len);
-	typedef bool(*HidD_GetManufacturerStringType)(HANDLE handle, PVOID buffer, uint32 buffer_len);
-	typedef bool(*HidD_GetProductStringType)(HANDLE handle, PVOID buffer, uint32 buffer_len);
-	typedef int32(*HidP_GetValueCapsType)(HIDP_REPORT_TYPE ReportType, PHIDP_VALUE_CAPS ValueCaps, uint16* ValueCapsLength, PHIDP_PREPARSED_DATA PreparsedData);
-	typedef int32(*HidP_GetButtonCapsType)(HIDP_REPORT_TYPE ReportType, PHIDP_BUTTON_CAPS ButtonCaps, uint16* ButtonCapsLength, PHIDP_PREPARSED_DATA PreparsedData);
-	typedef int32(*HidP_GetCapsType)(PHIDP_PREPARSED_DATA PreparsedData, PHIDP_CAPS Capabilities);
-	typedef int32(*HidP_GetUsagesType)(HIDP_REPORT_TYPE ReportType, uint16 UsagePage, uint16 LinkCollection, uint16* UsageList, uint32* UsageLength, PHIDP_PREPARSED_DATA PreparsedData, PCHAR Report, uint32 ReportLength);
-	typedef int32(*HidP_GetUsageValueType)(HIDP_REPORT_TYPE ReportType, uint16 UsagePage, uint16 LinkCollection, uint16 Usage, uint32* UsageValue, PHIDP_PREPARSED_DATA PreparsedData, PCHAR Report, uint32 ReportLength);
+	typedef bool(WINAPI *HidD_GetSerialNumberStringType)(HANDLE device, PVOID buffer, uint32 buffer_len);
+	typedef bool(WINAPI *HidD_GetManufacturerStringType)(HANDLE handle, PVOID buffer, uint32 buffer_len);
+	typedef bool(WINAPI *HidD_GetProductStringType)(HANDLE handle, PVOID buffer, uint32 buffer_len);
+	typedef int32(WINAPI *HidP_GetValueCapsType)(HIDP_REPORT_TYPE ReportType, PHIDP_VALUE_CAPS ValueCaps, uint16* ValueCapsLength, PHIDP_PREPARSED_DATA PreparsedData);
+	typedef int32(WINAPI *HidP_GetButtonCapsType)(HIDP_REPORT_TYPE ReportType, PHIDP_BUTTON_CAPS ButtonCaps, uint16* ButtonCapsLength, PHIDP_PREPARSED_DATA PreparsedData);
+	typedef int32(WINAPI *HidP_GetCapsType)(PHIDP_PREPARSED_DATA PreparsedData, PHIDP_CAPS Capabilities);
+	typedef int32(WINAPI *HidP_GetUsagesType)(HIDP_REPORT_TYPE ReportType, uint16 UsagePage, uint16 LinkCollection, uint16* UsageList, uint32* UsageLength, PHIDP_PREPARSED_DATA PreparsedData, PCHAR Report, uint32 ReportLength);
+	typedef int32(WINAPI *HidP_GetUsageValueType)(HIDP_REPORT_TYPE ReportType, uint16 UsagePage, uint16 LinkCollection, uint16 Usage, uint32* UsageValue, PHIDP_PREPARSED_DATA PreparsedData, PCHAR Report, uint32 ReportLength);
 	
 	FDLLPointers();
 	virtual ~FDLLPointers();
@@ -107,10 +107,11 @@ struct FAnalogData
 		, RangeMax(-1.f)
 		, Offset(0.f)
 		, bInverted(false)
+		, bGamepadStick(false)
 	{
 	}
 
-	FAnalogData(const int32 InIndex, const float InValue, const float InRangeMin, const float InRangeMax, const float InOffset, const bool bInInverted, const FName InKeyName)
+	FAnalogData(const int32 InIndex, const float InValue, const float InRangeMin, const float InRangeMax, const float InOffset, const bool bInInverted, const bool bInGamepadStick, const FName InKeyName)
 		: Index(InIndex)
 		, Value(InValue)
 		, PreviousValue(0.f)
@@ -118,6 +119,7 @@ struct FAnalogData
 		, RangeMax(InRangeMax)
 		, Offset(0.f)
 		, bInverted(bInInverted)
+		, bGamepadStick(bInGamepadStick)
 		, KeyName(InKeyName)
 	{
 	}
@@ -126,6 +128,13 @@ struct FAnalogData
 	float GetValue() const
 	{
 		const float Factor = 1.f / (RangeMax - RangeMin);
+		if (bGamepadStick)
+		{
+			// Need to do the re-centering before inverting
+			const float NormalizedValue = ((Value * Factor) - 0.5f) * 2 * (bInverted ? -1.0f : 1.0f);
+			return NormalizedValue + Offset;
+		}
+
 		const float NormalizedValue = (bInverted ? (Value * Factor * -1.f) : (Value * Factor));
 		return NormalizedValue + Offset;
 	}
@@ -156,6 +165,9 @@ struct FAnalogData
 
 	/* Is this axis inverted */
 	bool bInverted;
+
+	/* Is this axis centered on 0 instead of 0.5 */
+	bool bGamepadStick;
 	
 	/* Key name */
 	FName KeyName;
@@ -277,6 +289,9 @@ private:
 	/** Compare a RID_DEVICE_INFO structure with a FRawInputRegisteredDevice one. */
 	static bool CompareDeviceInfo(const RID_DEVICE_INFO& DeviceInfo, const FRawInputRegisteredDevice& OtherInfo);
 
+	/** Updates a raw device with connected device info, if null will reset data */
+	static void CopyConnectedDeviceInfo(FRawWindowsDeviceEntry& RegisteredDevice, const FConnectedDeviceInfo* ConnectedDevice);
+
 	/** Returns the handle to a registered device or INDEX_NONE if the device doesn't exist. */
 	int32 FindRegisteredDeviceHandle(const FRawInputRegisteredDevice& InDeviceData) const;
 
