Set defaults for FRegisteredDeviceInfo

diff --git a/Engine/Plugins/Experimental/RawInput/Source/RawInput/Public/RawInputFunctionLibrary.h b/Engine/Plugins/Experimental/RawInput/Source/RawInput/Public/RawInputFunctionLibrary.h
--- a/Engine/Plugins/Experimental/RawInput/Source/RawInput/Public/RawInputFunctionLibrary.h
+++ b/Engine/Plugins/Experimental/RawInput/Source/RawInput/Public/RawInputFunctionLibrary.h
@@ -80,15 +80,15 @@ struct RAWINPUT_API FRegisteredDeviceInfo
 	GENERATED_BODY()
 
 	UPROPERTY(BlueprintReadOnly, Category="RawInput")
-	int32 Handle;
+	int32 Handle = 0;
 
 	// Integer representation of the vendor ID (e.g. 0xC262 = 49762)
 	UPROPERTY(BlueprintReadOnly, Category="RawInput")
-	int32 VendorID;
+	int32 VendorID = 0;
 
 	// Integer representation of the product ID (e.g. 0xC262 = 49762)
 	UPROPERTY(BlueprintReadOnly, Category="RawInput")
-	int32 ProductID;
+	int32 ProductID = 0;
 
 	// Driver supplied device name
 	UPROPERTY(BlueprintReadOnly, Category="RawInput")
