RawInput: Improved logging quality
- Changed VID and PID printouts to be in hex
- Fixed a typo in VendorID
- Avoided printing out a warning about disconnected devices
- Downgraded the query results message to a Log from a Warning, and reworded it

diff --git a/Engine/Plugins/Experimental/RawInput/Source/RawInput/Private/Windows/RawInputWindows.cpp b/Engine/Plugins/Experimental/RawInput/Source/RawInput/Private/Windows/RawInputWindows.cpp
--- a/Engine/Plugins/Experimental/RawInput/Source/RawInput/Private/Windows/RawInputWindows.cpp
+++ b/Engine/Plugins/Experimental/RawInput/Source/RawInput/Private/Windows/RawInputWindows.cpp
@@ -181,7 +181,7 @@ int32 FRawInputWindows::RegisterInputDevice(const int32 DeviceType, const int32
 					FRawWindowsDeviceEntry& RegisteredDeviceInfo = RegisteredDeviceList[DeviceHandle];
 					CopyConnectedDeviceInfo(RegisteredDeviceInfo, &ConnectedDeviceInfo);
 
-					UE_LOG(LogRawInputWindows, Log, TEXT("VenderID:%x ProductID:%x"), RegisteredDeviceInfo.DeviceData.VendorID, RegisteredDeviceInfo.DeviceData.ProductID);
+					UE_LOG(LogRawInputWindows, Log, TEXT("VendorID:%04X ProductID:%04X"), RegisteredDeviceInfo.DeviceData.VendorID, RegisteredDeviceInfo.DeviceData.ProductID);
 
 					bWasConnected = true;
 					break;
@@ -688,7 +688,11 @@ void FRawInputWindows::QueryConnectedDevices()
 
 		if (GetRawInputDeviceInfoA(Device.hDevice, RIDI_DEVICENAME, DeviceNameBuffer.GetData(), &NameLen) == RAW_INPUT_ERROR)
 		{
-			UE_LOG(LogRawInputWindows, Warning, TEXT("Error reading device name"));
+			const DWORD LastErrorCode = GetLastError();
+			if (LastErrorCode != ERROR_FILE_NOT_FOUND)
+			{
+				UE_LOG(LogRawInputWindows, Warning, TEXT("Error reading device name (GetLastError = %d)"), LastErrorCode);
+			}
 			continue;
 		}
 
@@ -754,7 +758,7 @@ void FRawInputWindows::QueryConnectedDevices()
 		}
 	}
 
-	UE_LOG(LogRawInputWindows, Warning, TEXT("Found device %d devices"), ConnectedDeviceInfoList.Num());
+	UE_LOG(LogRawInputWindows, Log, TEXT("QueryConnectedDevices found %d devices"), ConnectedDeviceInfoList.Num());
 }
 
 FString FRawInputWindows::GetErrorString(const int32 StatusCode) const
@@ -885,7 +889,7 @@ void FRawInputWindows::ShowDeviceInfo(const FConnectedDeviceInfo& DeviceInfo) co
 		DeviceInfo.RIDDeviceInfo.keyboard.dwNumberOfKeysTotal);
 		break;
 	case RIM_TYPEHID:
-		UE_LOG(LogRawInputWindows, Verbose, TEXT("dwVendorId:%d, dwProductId:%d, dwVersionNumber:%d, usUsagePage:%d,usUsage:%d"),
+		UE_LOG(LogRawInputWindows, Verbose, TEXT("dwVendorId:%04X, dwProductId:%04X, dwVersionNumber:%d, usUsagePage:%d,usUsage:%d"),
 		DeviceInfo.RIDDeviceInfo.hid.dwVendorId,
 		DeviceInfo.RIDDeviceInfo.hid.dwProductId,
 		DeviceInfo.RIDDeviceInfo.hid.dwVersionNumber,
