Fix RawInputPlugin to properly support multiple concurrent devices. (problem reported on UDN)
Increase the number of Axis to 24 and the number of buttons to 96.

diff --git a/Engine/Plugins/Experimental/RawInput/Source/RawInput/Private/RawInput.cpp b/Engine/Plugins/Experimental/RawInput/Source/RawInput/Private/RawInput.cpp
--- a/Engine/Plugins/Experimental/RawInput/Source/RawInput/Private/RawInput.cpp
+++ b/Engine/Plugins/Experimental/RawInput/Source/RawInput/Private/RawInput.cpp
@@ -20,6 +20,22 @@ const FGamepadKeyNames::Type FRawInputKeyNames::GenericUSBController_Axis5( "Gen
 const FGamepadKeyNames::Type FRawInputKeyNames::GenericUSBController_Axis6( "GenericUSBController_Axis6" );
 const FGamepadKeyNames::Type FRawInputKeyNames::GenericUSBController_Axis7( "GenericUSBController_Axis7" );
 const FGamepadKeyNames::Type FRawInputKeyNames::GenericUSBController_Axis8( "GenericUSBController_Axis8" );
+const FGamepadKeyNames::Type FRawInputKeyNames::GenericUSBController_Axis9( "GenericUSBController_Axis9" );
+const FGamepadKeyNames::Type FRawInputKeyNames::GenericUSBController_Axis10( "GenericUSBController_Axis10" );
+const FGamepadKeyNames::Type FRawInputKeyNames::GenericUSBController_Axis11( "GenericUSBController_Axis11" );
+const FGamepadKeyNames::Type FRawInputKeyNames::GenericUSBController_Axis12( "GenericUSBController_Axis12" );
+const FGamepadKeyNames::Type FRawInputKeyNames::GenericUSBController_Axis13( "GenericUSBController_Axis13" );
+const FGamepadKeyNames::Type FRawInputKeyNames::GenericUSBController_Axis14( "GenericUSBController_Axis14" );
+const FGamepadKeyNames::Type FRawInputKeyNames::GenericUSBController_Axis15( "GenericUSBController_Axis15" );
+const FGamepadKeyNames::Type FRawInputKeyNames::GenericUSBController_Axis16( "GenericUSBController_Axis16" );
+const FGamepadKeyNames::Type FRawInputKeyNames::GenericUSBController_Axis17( "GenericUSBController_Axis17" );
+const FGamepadKeyNames::Type FRawInputKeyNames::GenericUSBController_Axis18( "GenericUSBController_Axis18" );
+const FGamepadKeyNames::Type FRawInputKeyNames::GenericUSBController_Axis19( "GenericUSBController_Axis19" );
+const FGamepadKeyNames::Type FRawInputKeyNames::GenericUSBController_Axis20( "GenericUSBController_Axis20" );
+const FGamepadKeyNames::Type FRawInputKeyNames::GenericUSBController_Axis21( "GenericUSBController_Axis21" );
+const FGamepadKeyNames::Type FRawInputKeyNames::GenericUSBController_Axis22( "GenericUSBController_Axis22" );
+const FGamepadKeyNames::Type FRawInputKeyNames::GenericUSBController_Axis23( "GenericUSBController_Axis23" );
+const FGamepadKeyNames::Type FRawInputKeyNames::GenericUSBController_Axis24( "GenericUSBController_Axis24" );
 
 const FGamepadKeyNames::Type FRawInputKeyNames::GenericUSBController_Button1("GenericUSBController_Button1");
 const FGamepadKeyNames::Type FRawInputKeyNames::GenericUSBController_Button2("GenericUSBController_Button2");
@@ -41,6 +57,82 @@ const FGamepadKeyNames::Type FRawInputKeyNames::GenericUSBController_Button17("G
 const FGamepadKeyNames::Type FRawInputKeyNames::GenericUSBController_Button18("GenericUSBController_Button18");
 const FGamepadKeyNames::Type FRawInputKeyNames::GenericUSBController_Button19("GenericUSBController_Button19");
 const FGamepadKeyNames::Type FRawInputKeyNames::GenericUSBController_Button20("GenericUSBController_Button20");
+const FGamepadKeyNames::Type FRawInputKeyNames::GenericUSBController_Button21("GenericUSBController_Button21");
+const FGamepadKeyNames::Type FRawInputKeyNames::GenericUSBController_Button22("GenericUSBController_Button22");
+const FGamepadKeyNames::Type FRawInputKeyNames::GenericUSBController_Button23("GenericUSBController_Button23");
+const FGamepadKeyNames::Type FRawInputKeyNames::GenericUSBController_Button24("GenericUSBController_Button24");
+const FGamepadKeyNames::Type FRawInputKeyNames::GenericUSBController_Button25("GenericUSBController_Button25");
+const FGamepadKeyNames::Type FRawInputKeyNames::GenericUSBController_Button26("GenericUSBController_Button26");
+const FGamepadKeyNames::Type FRawInputKeyNames::GenericUSBController_Button27("GenericUSBController_Button27");
+const FGamepadKeyNames::Type FRawInputKeyNames::GenericUSBController_Button28("GenericUSBController_Button28");
+const FGamepadKeyNames::Type FRawInputKeyNames::GenericUSBController_Button29("GenericUSBController_Button29");
+const FGamepadKeyNames::Type FRawInputKeyNames::GenericUSBController_Button30("GenericUSBController_Button30");
+const FGamepadKeyNames::Type FRawInputKeyNames::GenericUSBController_Button31("GenericUSBController_Button31");
+const FGamepadKeyNames::Type FRawInputKeyNames::GenericUSBController_Button32("GenericUSBController_Button32");
+const FGamepadKeyNames::Type FRawInputKeyNames::GenericUSBController_Button33("GenericUSBController_Button33");
+const FGamepadKeyNames::Type FRawInputKeyNames::GenericUSBController_Button34("GenericUSBController_Button34");
+const FGamepadKeyNames::Type FRawInputKeyNames::GenericUSBController_Button35("GenericUSBController_Button35");
+const FGamepadKeyNames::Type FRawInputKeyNames::GenericUSBController_Button36("GenericUSBController_Button36");
+const FGamepadKeyNames::Type FRawInputKeyNames::GenericUSBController_Button37("GenericUSBController_Button37");
+const FGamepadKeyNames::Type FRawInputKeyNames::GenericUSBController_Button38("GenericUSBController_Button38");
+const FGamepadKeyNames::Type FRawInputKeyNames::GenericUSBController_Button39("GenericUSBController_Button39");
+const FGamepadKeyNames::Type FRawInputKeyNames::GenericUSBController_Button40("GenericUSBController_Button40");
+const FGamepadKeyNames::Type FRawInputKeyNames::GenericUSBController_Button41("GenericUSBController_Button41");
+const FGamepadKeyNames::Type FRawInputKeyNames::GenericUSBController_Button42("GenericUSBController_Button42");
+const FGamepadKeyNames::Type FRawInputKeyNames::GenericUSBController_Button43("GenericUSBController_Button43");
+const FGamepadKeyNames::Type FRawInputKeyNames::GenericUSBController_Button44("GenericUSBController_Button44");
+const FGamepadKeyNames::Type FRawInputKeyNames::GenericUSBController_Button45("GenericUSBController_Button45");
+const FGamepadKeyNames::Type FRawInputKeyNames::GenericUSBController_Button46("GenericUSBController_Button46");
+const FGamepadKeyNames::Type FRawInputKeyNames::GenericUSBController_Button47("GenericUSBController_Button47");
+const FGamepadKeyNames::Type FRawInputKeyNames::GenericUSBController_Button48("GenericUSBController_Button48");
+const FGamepadKeyNames::Type FRawInputKeyNames::GenericUSBController_Button49("GenericUSBController_Button49");
+const FGamepadKeyNames::Type FRawInputKeyNames::GenericUSBController_Button50("GenericUSBController_Button50");
+const FGamepadKeyNames::Type FRawInputKeyNames::GenericUSBController_Button51("GenericUSBController_Button51");
+const FGamepadKeyNames::Type FRawInputKeyNames::GenericUSBController_Button52("GenericUSBController_Button52");
+const FGamepadKeyNames::Type FRawInputKeyNames::GenericUSBController_Button53("GenericUSBController_Button53");
+const FGamepadKeyNames::Type FRawInputKeyNames::GenericUSBController_Button54("GenericUSBController_Button54");
+const FGamepadKeyNames::Type FRawInputKeyNames::GenericUSBController_Button55("GenericUSBController_Button55");
+const FGamepadKeyNames::Type FRawInputKeyNames::GenericUSBController_Button56("GenericUSBController_Button56");
+const FGamepadKeyNames::Type FRawInputKeyNames::GenericUSBController_Button57("GenericUSBController_Button57");
+const FGamepadKeyNames::Type FRawInputKeyNames::GenericUSBController_Button58("GenericUSBController_Button58");
+const FGamepadKeyNames::Type FRawInputKeyNames::GenericUSBController_Button59("GenericUSBController_Button59");
+const FGamepadKeyNames::Type FRawInputKeyNames::GenericUSBController_Button60("GenericUSBController_Button60");
+const FGamepadKeyNames::Type FRawInputKeyNames::GenericUSBController_Button61("GenericUSBController_Button61");
+const FGamepadKeyNames::Type FRawInputKeyNames::GenericUSBController_Button62("GenericUSBController_Button62");
+const FGamepadKeyNames::Type FRawInputKeyNames::GenericUSBController_Button63("GenericUSBController_Button63");
+const FGamepadKeyNames::Type FRawInputKeyNames::GenericUSBController_Button64("GenericUSBController_Button64");
+const FGamepadKeyNames::Type FRawInputKeyNames::GenericUSBController_Button65("GenericUSBController_Button65");
+const FGamepadKeyNames::Type FRawInputKeyNames::GenericUSBController_Button66("GenericUSBController_Button66");
+const FGamepadKeyNames::Type FRawInputKeyNames::GenericUSBController_Button67("GenericUSBController_Button67");
+const FGamepadKeyNames::Type FRawInputKeyNames::GenericUSBController_Button68("GenericUSBController_Button68");
+const FGamepadKeyNames::Type FRawInputKeyNames::GenericUSBController_Button69("GenericUSBController_Button69");
+const FGamepadKeyNames::Type FRawInputKeyNames::GenericUSBController_Button70("GenericUSBController_Button70");
+const FGamepadKeyNames::Type FRawInputKeyNames::GenericUSBController_Button71("GenericUSBController_Button71");
+const FGamepadKeyNames::Type FRawInputKeyNames::GenericUSBController_Button72("GenericUSBController_Button72");
+const FGamepadKeyNames::Type FRawInputKeyNames::GenericUSBController_Button73("GenericUSBController_Button73");
+const FGamepadKeyNames::Type FRawInputKeyNames::GenericUSBController_Button74("GenericUSBController_Button74");
+const FGamepadKeyNames::Type FRawInputKeyNames::GenericUSBController_Button75("GenericUSBController_Button75");
+const FGamepadKeyNames::Type FRawInputKeyNames::GenericUSBController_Button76("GenericUSBController_Button76");
+const FGamepadKeyNames::Type FRawInputKeyNames::GenericUSBController_Button77("GenericUSBController_Button77");
+const FGamepadKeyNames::Type FRawInputKeyNames::GenericUSBController_Button78("GenericUSBController_Button78");
+const FGamepadKeyNames::Type FRawInputKeyNames::GenericUSBController_Button79("GenericUSBController_Button79");
+const FGamepadKeyNames::Type FRawInputKeyNames::GenericUSBController_Button80("GenericUSBController_Button80");
+const FGamepadKeyNames::Type FRawInputKeyNames::GenericUSBController_Button81("GenericUSBController_Button81");
+const FGamepadKeyNames::Type FRawInputKeyNames::GenericUSBController_Button82("GenericUSBController_Button82");
+const FGamepadKeyNames::Type FRawInputKeyNames::GenericUSBController_Button83("GenericUSBController_Button83");
+const FGamepadKeyNames::Type FRawInputKeyNames::GenericUSBController_Button84("GenericUSBController_Button84");
+const FGamepadKeyNames::Type FRawInputKeyNames::GenericUSBController_Button85("GenericUSBController_Button85");
+const FGamepadKeyNames::Type FRawInputKeyNames::GenericUSBController_Button86("GenericUSBController_Button86");
+const FGamepadKeyNames::Type FRawInputKeyNames::GenericUSBController_Button87("GenericUSBController_Button87");
+const FGamepadKeyNames::Type FRawInputKeyNames::GenericUSBController_Button88("GenericUSBController_Button88");
+const FGamepadKeyNames::Type FRawInputKeyNames::GenericUSBController_Button89("GenericUSBController_Button89");
+const FGamepadKeyNames::Type FRawInputKeyNames::GenericUSBController_Button90("GenericUSBController_Button90");
+const FGamepadKeyNames::Type FRawInputKeyNames::GenericUSBController_Button91("GenericUSBController_Button91");
+const FGamepadKeyNames::Type FRawInputKeyNames::GenericUSBController_Button92("GenericUSBController_Button92");
+const FGamepadKeyNames::Type FRawInputKeyNames::GenericUSBController_Button93("GenericUSBController_Button93");
+const FGamepadKeyNames::Type FRawInputKeyNames::GenericUSBController_Button94("GenericUSBController_Button94");
+const FGamepadKeyNames::Type FRawInputKeyNames::GenericUSBController_Button95("GenericUSBController_Button95");
+const FGamepadKeyNames::Type FRawInputKeyNames::GenericUSBController_Button96("GenericUSBController_Button96");
 
 // USB controller (Wheels, flight stick etc)
 const FKey FRawInputKeys::GenericUSBController_Axis1(FRawInputKeyNames::GenericUSBController_Axis1);
@@ -51,6 +143,22 @@ const FKey FRawInputKeys::GenericUSBController_Axis5(FRawInputKeyNames::GenericU
 const FKey FRawInputKeys::GenericUSBController_Axis6(FRawInputKeyNames::GenericUSBController_Axis6);
 const FKey FRawInputKeys::GenericUSBController_Axis7(FRawInputKeyNames::GenericUSBController_Axis7);
 const FKey FRawInputKeys::GenericUSBController_Axis8(FRawInputKeyNames::GenericUSBController_Axis8);
+const FKey FRawInputKeys::GenericUSBController_Axis9(FRawInputKeyNames::GenericUSBController_Axis9);
+const FKey FRawInputKeys::GenericUSBController_Axis10(FRawInputKeyNames::GenericUSBController_Axis10);
+const FKey FRawInputKeys::GenericUSBController_Axis11(FRawInputKeyNames::GenericUSBController_Axis11);
+const FKey FRawInputKeys::GenericUSBController_Axis12(FRawInputKeyNames::GenericUSBController_Axis12);
+const FKey FRawInputKeys::GenericUSBController_Axis13(FRawInputKeyNames::GenericUSBController_Axis13);
+const FKey FRawInputKeys::GenericUSBController_Axis14(FRawInputKeyNames::GenericUSBController_Axis14);
+const FKey FRawInputKeys::GenericUSBController_Axis15(FRawInputKeyNames::GenericUSBController_Axis15);
+const FKey FRawInputKeys::GenericUSBController_Axis16(FRawInputKeyNames::GenericUSBController_Axis16);
+const FKey FRawInputKeys::GenericUSBController_Axis17(FRawInputKeyNames::GenericUSBController_Axis17);
+const FKey FRawInputKeys::GenericUSBController_Axis18(FRawInputKeyNames::GenericUSBController_Axis18);
+const FKey FRawInputKeys::GenericUSBController_Axis19(FRawInputKeyNames::GenericUSBController_Axis19);
+const FKey FRawInputKeys::GenericUSBController_Axis20(FRawInputKeyNames::GenericUSBController_Axis20);
+const FKey FRawInputKeys::GenericUSBController_Axis21(FRawInputKeyNames::GenericUSBController_Axis21);
+const FKey FRawInputKeys::GenericUSBController_Axis22(FRawInputKeyNames::GenericUSBController_Axis22);
+const FKey FRawInputKeys::GenericUSBController_Axis23(FRawInputKeyNames::GenericUSBController_Axis23);
+const FKey FRawInputKeys::GenericUSBController_Axis24(FRawInputKeyNames::GenericUSBController_Axis24);
 
 const FKey FRawInputKeys::GenericUSBController_Button1(FRawInputKeyNames::GenericUSBController_Button1);
 const FKey FRawInputKeys::GenericUSBController_Button2(FRawInputKeyNames::GenericUSBController_Button2);
@@ -72,6 +180,82 @@ const FKey FRawInputKeys::GenericUSBController_Button17(FRawInputKeyNames::Gener
 const FKey FRawInputKeys::GenericUSBController_Button18(FRawInputKeyNames::GenericUSBController_Button18);
 const FKey FRawInputKeys::GenericUSBController_Button19(FRawInputKeyNames::GenericUSBController_Button19);
 const FKey FRawInputKeys::GenericUSBController_Button20(FRawInputKeyNames::GenericUSBController_Button20);
+const FKey FRawInputKeys::GenericUSBController_Button21(FRawInputKeyNames::GenericUSBController_Button21);
+const FKey FRawInputKeys::GenericUSBController_Button22(FRawInputKeyNames::GenericUSBController_Button22);
+const FKey FRawInputKeys::GenericUSBController_Button23(FRawInputKeyNames::GenericUSBController_Button23);
+const FKey FRawInputKeys::GenericUSBController_Button24(FRawInputKeyNames::GenericUSBController_Button24);
+const FKey FRawInputKeys::GenericUSBController_Button25(FRawInputKeyNames::GenericUSBController_Button25);
+const FKey FRawInputKeys::GenericUSBController_Button26(FRawInputKeyNames::GenericUSBController_Button26);
+const FKey FRawInputKeys::GenericUSBController_Button27(FRawInputKeyNames::GenericUSBController_Button27);
+const FKey FRawInputKeys::GenericUSBController_Button28(FRawInputKeyNames::GenericUSBController_Button28);
+const FKey FRawInputKeys::GenericUSBController_Button29(FRawInputKeyNames::GenericUSBController_Button29);
+const FKey FRawInputKeys::GenericUSBController_Button30(FRawInputKeyNames::GenericUSBController_Button30);
+const FKey FRawInputKeys::GenericUSBController_Button31(FRawInputKeyNames::GenericUSBController_Button31);
+const FKey FRawInputKeys::GenericUSBController_Button32(FRawInputKeyNames::GenericUSBController_Button32);
+const FKey FRawInputKeys::GenericUSBController_Button33(FRawInputKeyNames::GenericUSBController_Button33);
+const FKey FRawInputKeys::GenericUSBController_Button34(FRawInputKeyNames::GenericUSBController_Button34);
+const FKey FRawInputKeys::GenericUSBController_Button35(FRawInputKeyNames::GenericUSBController_Button35);
+const FKey FRawInputKeys::GenericUSBController_Button36(FRawInputKeyNames::GenericUSBController_Button36);
+const FKey FRawInputKeys::GenericUSBController_Button37(FRawInputKeyNames::GenericUSBController_Button37);
+const FKey FRawInputKeys::GenericUSBController_Button38(FRawInputKeyNames::GenericUSBController_Button38);
+const FKey FRawInputKeys::GenericUSBController_Button39(FRawInputKeyNames::GenericUSBController_Button39);
+const FKey FRawInputKeys::GenericUSBController_Button40(FRawInputKeyNames::GenericUSBController_Button40);
+const FKey FRawInputKeys::GenericUSBController_Button41(FRawInputKeyNames::GenericUSBController_Button41);
+const FKey FRawInputKeys::GenericUSBController_Button42(FRawInputKeyNames::GenericUSBController_Button42);
+const FKey FRawInputKeys::GenericUSBController_Button43(FRawInputKeyNames::GenericUSBController_Button43);
+const FKey FRawInputKeys::GenericUSBController_Button44(FRawInputKeyNames::GenericUSBController_Button44);
+const FKey FRawInputKeys::GenericUSBController_Button45(FRawInputKeyNames::GenericUSBController_Button45);
+const FKey FRawInputKeys::GenericUSBController_Button46(FRawInputKeyNames::GenericUSBController_Button46);
+const FKey FRawInputKeys::GenericUSBController_Button47(FRawInputKeyNames::GenericUSBController_Button47);
+const FKey FRawInputKeys::GenericUSBController_Button48(FRawInputKeyNames::GenericUSBController_Button48);
+const FKey FRawInputKeys::GenericUSBController_Button49(FRawInputKeyNames::GenericUSBController_Button49);
+const FKey FRawInputKeys::GenericUSBController_Button50(FRawInputKeyNames::GenericUSBController_Button50);
+const FKey FRawInputKeys::GenericUSBController_Button51(FRawInputKeyNames::GenericUSBController_Button51);
+const FKey FRawInputKeys::GenericUSBController_Button52(FRawInputKeyNames::GenericUSBController_Button52);
+const FKey FRawInputKeys::GenericUSBController_Button53(FRawInputKeyNames::GenericUSBController_Button53);
+const FKey FRawInputKeys::GenericUSBController_Button54(FRawInputKeyNames::GenericUSBController_Button54);
+const FKey FRawInputKeys::GenericUSBController_Button55(FRawInputKeyNames::GenericUSBController_Button55);
+const FKey FRawInputKeys::GenericUSBController_Button56(FRawInputKeyNames::GenericUSBController_Button56);
+const FKey FRawInputKeys::GenericUSBController_Button57(FRawInputKeyNames::GenericUSBController_Button57);
+const FKey FRawInputKeys::GenericUSBController_Button58(FRawInputKeyNames::GenericUSBController_Button58);
+const FKey FRawInputKeys::GenericUSBController_Button59(FRawInputKeyNames::GenericUSBController_Button59);
+const FKey FRawInputKeys::GenericUSBController_Button60(FRawInputKeyNames::GenericUSBController_Button60);
+const FKey FRawInputKeys::GenericUSBController_Button61(FRawInputKeyNames::GenericUSBController_Button61);
+const FKey FRawInputKeys::GenericUSBController_Button62(FRawInputKeyNames::GenericUSBController_Button62);
+const FKey FRawInputKeys::GenericUSBController_Button63(FRawInputKeyNames::GenericUSBController_Button63);
+const FKey FRawInputKeys::GenericUSBController_Button64(FRawInputKeyNames::GenericUSBController_Button64);
+const FKey FRawInputKeys::GenericUSBController_Button65(FRawInputKeyNames::GenericUSBController_Button65);
+const FKey FRawInputKeys::GenericUSBController_Button66(FRawInputKeyNames::GenericUSBController_Button66);
+const FKey FRawInputKeys::GenericUSBController_Button67(FRawInputKeyNames::GenericUSBController_Button67);
+const FKey FRawInputKeys::GenericUSBController_Button68(FRawInputKeyNames::GenericUSBController_Button68);
+const FKey FRawInputKeys::GenericUSBController_Button69(FRawInputKeyNames::GenericUSBController_Button69);
+const FKey FRawInputKeys::GenericUSBController_Button70(FRawInputKeyNames::GenericUSBController_Button70);
+const FKey FRawInputKeys::GenericUSBController_Button71(FRawInputKeyNames::GenericUSBController_Button71);
+const FKey FRawInputKeys::GenericUSBController_Button72(FRawInputKeyNames::GenericUSBController_Button72);
+const FKey FRawInputKeys::GenericUSBController_Button73(FRawInputKeyNames::GenericUSBController_Button73);
+const FKey FRawInputKeys::GenericUSBController_Button74(FRawInputKeyNames::GenericUSBController_Button74);
+const FKey FRawInputKeys::GenericUSBController_Button75(FRawInputKeyNames::GenericUSBController_Button75);
+const FKey FRawInputKeys::GenericUSBController_Button76(FRawInputKeyNames::GenericUSBController_Button76);
+const FKey FRawInputKeys::GenericUSBController_Button77(FRawInputKeyNames::GenericUSBController_Button77);
+const FKey FRawInputKeys::GenericUSBController_Button78(FRawInputKeyNames::GenericUSBController_Button78);
+const FKey FRawInputKeys::GenericUSBController_Button79(FRawInputKeyNames::GenericUSBController_Button79);
+const FKey FRawInputKeys::GenericUSBController_Button80(FRawInputKeyNames::GenericUSBController_Button80);
+const FKey FRawInputKeys::GenericUSBController_Button81(FRawInputKeyNames::GenericUSBController_Button81);
+const FKey FRawInputKeys::GenericUSBController_Button82(FRawInputKeyNames::GenericUSBController_Button82);
+const FKey FRawInputKeys::GenericUSBController_Button83(FRawInputKeyNames::GenericUSBController_Button83);
+const FKey FRawInputKeys::GenericUSBController_Button84(FRawInputKeyNames::GenericUSBController_Button84);
+const FKey FRawInputKeys::GenericUSBController_Button85(FRawInputKeyNames::GenericUSBController_Button85);
+const FKey FRawInputKeys::GenericUSBController_Button86(FRawInputKeyNames::GenericUSBController_Button86);
+const FKey FRawInputKeys::GenericUSBController_Button87(FRawInputKeyNames::GenericUSBController_Button87);
+const FKey FRawInputKeys::GenericUSBController_Button88(FRawInputKeyNames::GenericUSBController_Button88);
+const FKey FRawInputKeys::GenericUSBController_Button89(FRawInputKeyNames::GenericUSBController_Button89);
+const FKey FRawInputKeys::GenericUSBController_Button90(FRawInputKeyNames::GenericUSBController_Button90);
+const FKey FRawInputKeys::GenericUSBController_Button91(FRawInputKeyNames::GenericUSBController_Button91);
+const FKey FRawInputKeys::GenericUSBController_Button92(FRawInputKeyNames::GenericUSBController_Button92);
+const FKey FRawInputKeys::GenericUSBController_Button93(FRawInputKeyNames::GenericUSBController_Button93);
+const FKey FRawInputKeys::GenericUSBController_Button94(FRawInputKeyNames::GenericUSBController_Button94);
+const FKey FRawInputKeys::GenericUSBController_Button95(FRawInputKeyNames::GenericUSBController_Button95);
+const FKey FRawInputKeys::GenericUSBController_Button96(FRawInputKeyNames::GenericUSBController_Button96);
 
 IRawInput::IRawInput(const TSharedRef< FGenericApplicationMessageHandler >& InMessageHandler)
 	: MessageHandler(InMessageHandler)
@@ -102,6 +286,22 @@ void FRawInputPlugin::StartupModule()
 	EKeys::AddKey(FKeyDetails(FRawInputKeys::GenericUSBController_Axis6, LOCTEXT("GenericUSBController_Axis6", "GenericUSBController Axis 6"), FKeyDetails::GamepadKey, NAME_GenericUSBController));
 	EKeys::AddKey(FKeyDetails(FRawInputKeys::GenericUSBController_Axis7, LOCTEXT("GenericUSBController_Axis7", "GenericUSBController Axis 7"), FKeyDetails::GamepadKey, NAME_GenericUSBController));
 	EKeys::AddKey(FKeyDetails(FRawInputKeys::GenericUSBController_Axis8, LOCTEXT("GenericUSBController_Axis8", "GenericUSBController Axis 8"), FKeyDetails::GamepadKey, NAME_GenericUSBController));
+	EKeys::AddKey(FKeyDetails(FRawInputKeys::GenericUSBController_Axis9, LOCTEXT("GenericUSBController_Axis9", "GenericUSBController Axis 9"), FKeyDetails::GamepadKey, NAME_GenericUSBController));
+	EKeys::AddKey(FKeyDetails(FRawInputKeys::GenericUSBController_Axis10, LOCTEXT("GenericUSBController_Axis10", "GenericUSBController Axis 10"), FKeyDetails::GamepadKey, NAME_GenericUSBController));
+	EKeys::AddKey(FKeyDetails(FRawInputKeys::GenericUSBController_Axis11, LOCTEXT("GenericUSBController_Axis11", "GenericUSBController Axis 11"), FKeyDetails::GamepadKey, NAME_GenericUSBController));
+	EKeys::AddKey(FKeyDetails(FRawInputKeys::GenericUSBController_Axis12, LOCTEXT("GenericUSBController_Axis12", "GenericUSBController Axis 12"), FKeyDetails::GamepadKey, NAME_GenericUSBController));
+	EKeys::AddKey(FKeyDetails(FRawInputKeys::GenericUSBController_Axis13, LOCTEXT("GenericUSBController_Axis13", "GenericUSBController Axis 13"), FKeyDetails::GamepadKey, NAME_GenericUSBController));
+	EKeys::AddKey(FKeyDetails(FRawInputKeys::GenericUSBController_Axis14, LOCTEXT("GenericUSBController_Axis14", "GenericUSBController Axis 14"), FKeyDetails::GamepadKey, NAME_GenericUSBController));
+	EKeys::AddKey(FKeyDetails(FRawInputKeys::GenericUSBController_Axis15, LOCTEXT("GenericUSBController_Axis15", "GenericUSBController Axis 15"), FKeyDetails::GamepadKey, NAME_GenericUSBController));
+	EKeys::AddKey(FKeyDetails(FRawInputKeys::GenericUSBController_Axis16, LOCTEXT("GenericUSBController_Axis16", "GenericUSBController Axis 16"), FKeyDetails::GamepadKey, NAME_GenericUSBController));
+	EKeys::AddKey(FKeyDetails(FRawInputKeys::GenericUSBController_Axis17, LOCTEXT("GenericUSBController_Axis17", "GenericUSBController Axis 17"), FKeyDetails::GamepadKey, NAME_GenericUSBController));
+	EKeys::AddKey(FKeyDetails(FRawInputKeys::GenericUSBController_Axis18, LOCTEXT("GenericUSBController_Axis18", "GenericUSBController Axis 18"), FKeyDetails::GamepadKey, NAME_GenericUSBController));
+	EKeys::AddKey(FKeyDetails(FRawInputKeys::GenericUSBController_Axis19, LOCTEXT("GenericUSBController_Axis19", "GenericUSBController Axis 19"), FKeyDetails::GamepadKey, NAME_GenericUSBController));
+	EKeys::AddKey(FKeyDetails(FRawInputKeys::GenericUSBController_Axis20, LOCTEXT("GenericUSBController_Axis20", "GenericUSBController Axis 20"), FKeyDetails::GamepadKey, NAME_GenericUSBController));
+	EKeys::AddKey(FKeyDetails(FRawInputKeys::GenericUSBController_Axis21, LOCTEXT("GenericUSBController_Axis21", "GenericUSBController Axis 21"), FKeyDetails::GamepadKey, NAME_GenericUSBController));
+	EKeys::AddKey(FKeyDetails(FRawInputKeys::GenericUSBController_Axis22, LOCTEXT("GenericUSBController_Axis22", "GenericUSBController Axis 22"), FKeyDetails::GamepadKey, NAME_GenericUSBController));
+	EKeys::AddKey(FKeyDetails(FRawInputKeys::GenericUSBController_Axis23, LOCTEXT("GenericUSBController_Axis23", "GenericUSBController Axis 23"), FKeyDetails::GamepadKey, NAME_GenericUSBController));
+	EKeys::AddKey(FKeyDetails(FRawInputKeys::GenericUSBController_Axis24, LOCTEXT("GenericUSBController_Axis24", "GenericUSBController Axis 24"), FKeyDetails::GamepadKey, NAME_GenericUSBController));
 
 	EKeys::AddKey(FKeyDetails(FRawInputKeys::GenericUSBController_Button1, LOCTEXT("GenericUSBController_Button1", "GenericUSBController Button 1"), FKeyDetails::GamepadKey, NAME_GenericUSBController));
 	EKeys::AddKey(FKeyDetails(FRawInputKeys::GenericUSBController_Button2, LOCTEXT("GenericUSBController_Button2", "GenericUSBController Button 2"), FKeyDetails::GamepadKey, NAME_GenericUSBController));
@@ -123,6 +323,82 @@ void FRawInputPlugin::StartupModule()
 	EKeys::AddKey(FKeyDetails(FRawInputKeys::GenericUSBController_Button18, LOCTEXT("GenericUSBController_Button18", "GenericUSBController Button 18"), FKeyDetails::GamepadKey, NAME_GenericUSBController));
 	EKeys::AddKey(FKeyDetails(FRawInputKeys::GenericUSBController_Button19, LOCTEXT("GenericUSBController_Button19", "GenericUSBController Button 19"), FKeyDetails::GamepadKey, NAME_GenericUSBController));
 	EKeys::AddKey(FKeyDetails(FRawInputKeys::GenericUSBController_Button20, LOCTEXT("GenericUSBController_Button20", "GenericUSBController Button 20"), FKeyDetails::GamepadKey, NAME_GenericUSBController));
+	EKeys::AddKey(FKeyDetails(FRawInputKeys::GenericUSBController_Button21, LOCTEXT("GenericUSBController_Button21", "GenericUSBController Button 21"), FKeyDetails::GamepadKey, NAME_GenericUSBController));
+	EKeys::AddKey(FKeyDetails(FRawInputKeys::GenericUSBController_Button22, LOCTEXT("GenericUSBController_Button22", "GenericUSBController Button 22"), FKeyDetails::GamepadKey, NAME_GenericUSBController));
+	EKeys::AddKey(FKeyDetails(FRawInputKeys::GenericUSBController_Button23, LOCTEXT("GenericUSBController_Button23", "GenericUSBController Button 23"), FKeyDetails::GamepadKey, NAME_GenericUSBController));
+	EKeys::AddKey(FKeyDetails(FRawInputKeys::GenericUSBController_Button24, LOCTEXT("GenericUSBController_Button24", "GenericUSBController Button 24"), FKeyDetails::GamepadKey, NAME_GenericUSBController));
+	EKeys::AddKey(FKeyDetails(FRawInputKeys::GenericUSBController_Button25, LOCTEXT("GenericUSBController_Button25", "GenericUSBController Button 25"), FKeyDetails::GamepadKey, NAME_GenericUSBController));
+	EKeys::AddKey(FKeyDetails(FRawInputKeys::GenericUSBController_Button26, LOCTEXT("GenericUSBController_Button26", "GenericUSBController Button 26"), FKeyDetails::GamepadKey, NAME_GenericUSBController));
+	EKeys::AddKey(FKeyDetails(FRawInputKeys::GenericUSBController_Button27, LOCTEXT("GenericUSBController_Button27", "GenericUSBController Button 27"), FKeyDetails::GamepadKey, NAME_GenericUSBController));
+	EKeys::AddKey(FKeyDetails(FRawInputKeys::GenericUSBController_Button28, LOCTEXT("GenericUSBController_Button28", "GenericUSBController Button 28"), FKeyDetails::GamepadKey, NAME_GenericUSBController));
+	EKeys::AddKey(FKeyDetails(FRawInputKeys::GenericUSBController_Button29, LOCTEXT("GenericUSBController_Button29", "GenericUSBController Button 29"), FKeyDetails::GamepadKey, NAME_GenericUSBController));
+	EKeys::AddKey(FKeyDetails(FRawInputKeys::GenericUSBController_Button30, LOCTEXT("GenericUSBController_Button30", "GenericUSBController Button 30"), FKeyDetails::GamepadKey, NAME_GenericUSBController));
+	EKeys::AddKey(FKeyDetails(FRawInputKeys::GenericUSBController_Button31, LOCTEXT("GenericUSBController_Button31", "GenericUSBController Button 31"), FKeyDetails::GamepadKey, NAME_GenericUSBController));
+	EKeys::AddKey(FKeyDetails(FRawInputKeys::GenericUSBController_Button32, LOCTEXT("GenericUSBController_Button32", "GenericUSBController Button 32"), FKeyDetails::GamepadKey, NAME_GenericUSBController));
+	EKeys::AddKey(FKeyDetails(FRawInputKeys::GenericUSBController_Button33, LOCTEXT("GenericUSBController_Button33", "GenericUSBController Button 33"), FKeyDetails::GamepadKey, NAME_GenericUSBController));
+	EKeys::AddKey(FKeyDetails(FRawInputKeys::GenericUSBController_Button34, LOCTEXT("GenericUSBController_Button34", "GenericUSBController Button 34"), FKeyDetails::GamepadKey, NAME_GenericUSBController));
+	EKeys::AddKey(FKeyDetails(FRawInputKeys::GenericUSBController_Button35, LOCTEXT("GenericUSBController_Button35", "GenericUSBController Button 35"), FKeyDetails::GamepadKey, NAME_GenericUSBController));
+	EKeys::AddKey(FKeyDetails(FRawInputKeys::GenericUSBController_Button36, LOCTEXT("GenericUSBController_Button36", "GenericUSBController Button 36"), FKeyDetails::GamepadKey, NAME_GenericUSBController));
+	EKeys::AddKey(FKeyDetails(FRawInputKeys::GenericUSBController_Button37, LOCTEXT("GenericUSBController_Button37", "GenericUSBController Button 37"), FKeyDetails::GamepadKey, NAME_GenericUSBController));
+	EKeys::AddKey(FKeyDetails(FRawInputKeys::GenericUSBController_Button38, LOCTEXT("GenericUSBController_Button38", "GenericUSBController Button 38"), FKeyDetails::GamepadKey, NAME_GenericUSBController));
+	EKeys::AddKey(FKeyDetails(FRawInputKeys::GenericUSBController_Button39, LOCTEXT("GenericUSBController_Button39", "GenericUSBController Button 39"), FKeyDetails::GamepadKey, NAME_GenericUSBController));
+	EKeys::AddKey(FKeyDetails(FRawInputKeys::GenericUSBController_Button40, LOCTEXT("GenericUSBController_Button40", "GenericUSBController Button 40"), FKeyDetails::GamepadKey, NAME_GenericUSBController));
+	EKeys::AddKey(FKeyDetails(FRawInputKeys::GenericUSBController_Button41, LOCTEXT("GenericUSBController_Button41", "GenericUSBController Button 41"), FKeyDetails::GamepadKey, NAME_GenericUSBController));
+	EKeys::AddKey(FKeyDetails(FRawInputKeys::GenericUSBController_Button42, LOCTEXT("GenericUSBController_Button42", "GenericUSBController Button 42"), FKeyDetails::GamepadKey, NAME_GenericUSBController));
+	EKeys::AddKey(FKeyDetails(FRawInputKeys::GenericUSBController_Button43, LOCTEXT("GenericUSBController_Button43", "GenericUSBController Button 43"), FKeyDetails::GamepadKey, NAME_GenericUSBController));
+	EKeys::AddKey(FKeyDetails(FRawInputKeys::GenericUSBController_Button44, LOCTEXT("GenericUSBController_Button44", "GenericUSBController Button 44"), FKeyDetails::GamepadKey, NAME_GenericUSBController));
+	EKeys::AddKey(FKeyDetails(FRawInputKeys::GenericUSBController_Button45, LOCTEXT("GenericUSBController_Button45", "GenericUSBController Button 45"), FKeyDetails::GamepadKey, NAME_GenericUSBController));
+	EKeys::AddKey(FKeyDetails(FRawInputKeys::GenericUSBController_Button46, LOCTEXT("GenericUSBController_Button46", "GenericUSBController Button 46"), FKeyDetails::GamepadKey, NAME_GenericUSBController));
+	EKeys::AddKey(FKeyDetails(FRawInputKeys::GenericUSBController_Button47, LOCTEXT("GenericUSBController_Button47", "GenericUSBController Button 47"), FKeyDetails::GamepadKey, NAME_GenericUSBController));
+	EKeys::AddKey(FKeyDetails(FRawInputKeys::GenericUSBController_Button48, LOCTEXT("GenericUSBController_Button48", "GenericUSBController Button 48"), FKeyDetails::GamepadKey, NAME_GenericUSBController));
+	EKeys::AddKey(FKeyDetails(FRawInputKeys::GenericUSBController_Button49, LOCTEXT("GenericUSBController_Button49", "GenericUSBController Button 49"), FKeyDetails::GamepadKey, NAME_GenericUSBController));
+	EKeys::AddKey(FKeyDetails(FRawInputKeys::GenericUSBController_Button50, LOCTEXT("GenericUSBController_Button50", "GenericUSBController Button 50"), FKeyDetails::GamepadKey, NAME_GenericUSBController));
+	EKeys::AddKey(FKeyDetails(FRawInputKeys::GenericUSBController_Button51, LOCTEXT("GenericUSBController_Button51", "GenericUSBController Button 51"), FKeyDetails::GamepadKey, NAME_GenericUSBController));
+	EKeys::AddKey(FKeyDetails(FRawInputKeys::GenericUSBController_Button52, LOCTEXT("GenericUSBController_Button52", "GenericUSBController Button 52"), FKeyDetails::GamepadKey, NAME_GenericUSBController));
+	EKeys::AddKey(FKeyDetails(FRawInputKeys::GenericUSBController_Button53, LOCTEXT("GenericUSBController_Button53", "GenericUSBController Button 53"), FKeyDetails::GamepadKey, NAME_GenericUSBController));
+	EKeys::AddKey(FKeyDetails(FRawInputKeys::GenericUSBController_Button54, LOCTEXT("GenericUSBController_Button54", "GenericUSBController Button 54"), FKeyDetails::GamepadKey, NAME_GenericUSBController));
+	EKeys::AddKey(FKeyDetails(FRawInputKeys::GenericUSBController_Button55, LOCTEXT("GenericUSBController_Button55", "GenericUSBController Button 55"), FKeyDetails::GamepadKey, NAME_GenericUSBController));
+	EKeys::AddKey(FKeyDetails(FRawInputKeys::GenericUSBController_Button56, LOCTEXT("GenericUSBController_Button56", "GenericUSBController Button 56"), FKeyDetails::GamepadKey, NAME_GenericUSBController));
+	EKeys::AddKey(FKeyDetails(FRawInputKeys::GenericUSBController_Button57, LOCTEXT("GenericUSBController_Button57", "GenericUSBController Button 57"), FKeyDetails::GamepadKey, NAME_GenericUSBController));
+	EKeys::AddKey(FKeyDetails(FRawInputKeys::GenericUSBController_Button58, LOCTEXT("GenericUSBController_Button58", "GenericUSBController Button 58"), FKeyDetails::GamepadKey, NAME_GenericUSBController));
+	EKeys::AddKey(FKeyDetails(FRawInputKeys::GenericUSBController_Button59, LOCTEXT("GenericUSBController_Button59", "GenericUSBController Button 59"), FKeyDetails::GamepadKey, NAME_GenericUSBController));
+	EKeys::AddKey(FKeyDetails(FRawInputKeys::GenericUSBController_Button60, LOCTEXT("GenericUSBController_Button60", "GenericUSBController Button 60"), FKeyDetails::GamepadKey, NAME_GenericUSBController));
+	EKeys::AddKey(FKeyDetails(FRawInputKeys::GenericUSBController_Button61, LOCTEXT("GenericUSBController_Button61", "GenericUSBController Button 61"), FKeyDetails::GamepadKey, NAME_GenericUSBController));
+	EKeys::AddKey(FKeyDetails(FRawInputKeys::GenericUSBController_Button62, LOCTEXT("GenericUSBController_Button62", "GenericUSBController Button 62"), FKeyDetails::GamepadKey, NAME_GenericUSBController));
+	EKeys::AddKey(FKeyDetails(FRawInputKeys::GenericUSBController_Button63, LOCTEXT("GenericUSBController_Button63", "GenericUSBController Button 63"), FKeyDetails::GamepadKey, NAME_GenericUSBController));
+	EKeys::AddKey(FKeyDetails(FRawInputKeys::GenericUSBController_Button64, LOCTEXT("GenericUSBController_Button64", "GenericUSBController Button 64"), FKeyDetails::GamepadKey, NAME_GenericUSBController));
+	EKeys::AddKey(FKeyDetails(FRawInputKeys::GenericUSBController_Button65, LOCTEXT("GenericUSBController_Button65", "GenericUSBController Button 65"), FKeyDetails::GamepadKey, NAME_GenericUSBController));
+	EKeys::AddKey(FKeyDetails(FRawInputKeys::GenericUSBController_Button66, LOCTEXT("GenericUSBController_Button66", "GenericUSBController Button 66"), FKeyDetails::GamepadKey, NAME_GenericUSBController));
+	EKeys::AddKey(FKeyDetails(FRawInputKeys::GenericUSBController_Button67, LOCTEXT("GenericUSBController_Button67", "GenericUSBController Button 67"), FKeyDetails::GamepadKey, NAME_GenericUSBController));
+	EKeys::AddKey(FKeyDetails(FRawInputKeys::GenericUSBController_Button68, LOCTEXT("GenericUSBController_Button68", "GenericUSBController Button 68"), FKeyDetails::GamepadKey, NAME_GenericUSBController));
+	EKeys::AddKey(FKeyDetails(FRawInputKeys::GenericUSBController_Button69, LOCTEXT("GenericUSBController_Button69", "GenericUSBController Button 69"), FKeyDetails::GamepadKey, NAME_GenericUSBController));
+	EKeys::AddKey(FKeyDetails(FRawInputKeys::GenericUSBController_Button70, LOCTEXT("GenericUSBController_Button70", "GenericUSBController Button 70"), FKeyDetails::GamepadKey, NAME_GenericUSBController));
+	EKeys::AddKey(FKeyDetails(FRawInputKeys::GenericUSBController_Button71, LOCTEXT("GenericUSBController_Button71", "GenericUSBController Button 71"), FKeyDetails::GamepadKey, NAME_GenericUSBController));
+	EKeys::AddKey(FKeyDetails(FRawInputKeys::GenericUSBController_Button72, LOCTEXT("GenericUSBController_Button72", "GenericUSBController Button 72"), FKeyDetails::GamepadKey, NAME_GenericUSBController));
+	EKeys::AddKey(FKeyDetails(FRawInputKeys::GenericUSBController_Button73, LOCTEXT("GenericUSBController_Button73", "GenericUSBController Button 73"), FKeyDetails::GamepadKey, NAME_GenericUSBController));
+	EKeys::AddKey(FKeyDetails(FRawInputKeys::GenericUSBController_Button74, LOCTEXT("GenericUSBController_Button74", "GenericUSBController Button 74"), FKeyDetails::GamepadKey, NAME_GenericUSBController));
+	EKeys::AddKey(FKeyDetails(FRawInputKeys::GenericUSBController_Button75, LOCTEXT("GenericUSBController_Button75", "GenericUSBController Button 75"), FKeyDetails::GamepadKey, NAME_GenericUSBController));
+	EKeys::AddKey(FKeyDetails(FRawInputKeys::GenericUSBController_Button76, LOCTEXT("GenericUSBController_Button76", "GenericUSBController Button 76"), FKeyDetails::GamepadKey, NAME_GenericUSBController));
+	EKeys::AddKey(FKeyDetails(FRawInputKeys::GenericUSBController_Button77, LOCTEXT("GenericUSBController_Button77", "GenericUSBController Button 77"), FKeyDetails::GamepadKey, NAME_GenericUSBController));
+	EKeys::AddKey(FKeyDetails(FRawInputKeys::GenericUSBController_Button78, LOCTEXT("GenericUSBController_Button78", "GenericUSBController Button 78"), FKeyDetails::GamepadKey, NAME_GenericUSBController));
+	EKeys::AddKey(FKeyDetails(FRawInputKeys::GenericUSBController_Button79, LOCTEXT("GenericUSBController_Button79", "GenericUSBController Button 79"), FKeyDetails::GamepadKey, NAME_GenericUSBController));
+	EKeys::AddKey(FKeyDetails(FRawInputKeys::GenericUSBController_Button80, LOCTEXT("GenericUSBController_Button80", "GenericUSBController Button 80"), FKeyDetails::GamepadKey, NAME_GenericUSBController));
+	EKeys::AddKey(FKeyDetails(FRawInputKeys::GenericUSBController_Button81, LOCTEXT("GenericUSBController_Button81", "GenericUSBController Button 81"), FKeyDetails::GamepadKey, NAME_GenericUSBController));
+	EKeys::AddKey(FKeyDetails(FRawInputKeys::GenericUSBController_Button82, LOCTEXT("GenericUSBController_Button82", "GenericUSBController Button 82"), FKeyDetails::GamepadKey, NAME_GenericUSBController));
+	EKeys::AddKey(FKeyDetails(FRawInputKeys::GenericUSBController_Button83, LOCTEXT("GenericUSBController_Button83", "GenericUSBController Button 83"), FKeyDetails::GamepadKey, NAME_GenericUSBController));
+	EKeys::AddKey(FKeyDetails(FRawInputKeys::GenericUSBController_Button84, LOCTEXT("GenericUSBController_Button84", "GenericUSBController Button 84"), FKeyDetails::GamepadKey, NAME_GenericUSBController));
+	EKeys::AddKey(FKeyDetails(FRawInputKeys::GenericUSBController_Button85, LOCTEXT("GenericUSBController_Button85", "GenericUSBController Button 85"), FKeyDetails::GamepadKey, NAME_GenericUSBController));
+	EKeys::AddKey(FKeyDetails(FRawInputKeys::GenericUSBController_Button86, LOCTEXT("GenericUSBController_Button86", "GenericUSBController Button 86"), FKeyDetails::GamepadKey, NAME_GenericUSBController));
+	EKeys::AddKey(FKeyDetails(FRawInputKeys::GenericUSBController_Button87, LOCTEXT("GenericUSBController_Button87", "GenericUSBController Button 87"), FKeyDetails::GamepadKey, NAME_GenericUSBController));
+	EKeys::AddKey(FKeyDetails(FRawInputKeys::GenericUSBController_Button88, LOCTEXT("GenericUSBController_Button88", "GenericUSBController Button 88"), FKeyDetails::GamepadKey, NAME_GenericUSBController));
+	EKeys::AddKey(FKeyDetails(FRawInputKeys::GenericUSBController_Button89, LOCTEXT("GenericUSBController_Button89", "GenericUSBController Button 89"), FKeyDetails::GamepadKey, NAME_GenericUSBController));
+	EKeys::AddKey(FKeyDetails(FRawInputKeys::GenericUSBController_Button90, LOCTEXT("GenericUSBController_Button90", "GenericUSBController Button 90"), FKeyDetails::GamepadKey, NAME_GenericUSBController));
+	EKeys::AddKey(FKeyDetails(FRawInputKeys::GenericUSBController_Button91, LOCTEXT("GenericUSBController_Button91", "GenericUSBController Button 91"), FKeyDetails::GamepadKey, NAME_GenericUSBController));
+	EKeys::AddKey(FKeyDetails(FRawInputKeys::GenericUSBController_Button92, LOCTEXT("GenericUSBController_Button92", "GenericUSBController Button 92"), FKeyDetails::GamepadKey, NAME_GenericUSBController));
+	EKeys::AddKey(FKeyDetails(FRawInputKeys::GenericUSBController_Button93, LOCTEXT("GenericUSBController_Button93", "GenericUSBController Button 93"), FKeyDetails::GamepadKey, NAME_GenericUSBController));
+	EKeys::AddKey(FKeyDetails(FRawInputKeys::GenericUSBController_Button94, LOCTEXT("GenericUSBController_Button94", "GenericUSBController Button 94"), FKeyDetails::GamepadKey, NAME_GenericUSBController));
+	EKeys::AddKey(FKeyDetails(FRawInputKeys::GenericUSBController_Button95, LOCTEXT("GenericUSBController_Button95", "GenericUSBController Button 95"), FKeyDetails::GamepadKey, NAME_GenericUSBController));
+	EKeys::AddKey(FKeyDetails(FRawInputKeys::GenericUSBController_Button96, LOCTEXT("GenericUSBController_Button96", "GenericUSBController Button 96"), FKeyDetails::GamepadKey, NAME_GenericUSBController));
 }
 
 IMPLEMENT_MODULE( FRawInputPlugin, RawInput)
diff --git a/Engine/Plugins/Experimental/RawInput/Source/RawInput/Private/RawInputSettings.cpp b/Engine/Plugins/Experimental/RawInput/Source/RawInput/Private/RawInputSettings.cpp
--- a/Engine/Plugins/Experimental/RawInput/Source/RawInput/Private/RawInputSettings.cpp
+++ b/Engine/Plugins/Experimental/RawInput/Source/RawInput/Private/RawInputSettings.cpp
@@ -32,6 +32,82 @@ FRawInputDeviceConfiguration::FRawInputDeviceConfiguration()
 	ButtonProperties[17].Key = FRawInputKeys::GenericUSBController_Button18;
 	ButtonProperties[18].Key = FRawInputKeys::GenericUSBController_Button19;
 	ButtonProperties[19].Key = FRawInputKeys::GenericUSBController_Button20;
+	ButtonProperties[20].Key = FRawInputKeys::GenericUSBController_Button21;
+	ButtonProperties[21].Key = FRawInputKeys::GenericUSBController_Button22;
+	ButtonProperties[22].Key = FRawInputKeys::GenericUSBController_Button23;
+	ButtonProperties[23].Key = FRawInputKeys::GenericUSBController_Button24;
+	ButtonProperties[24].Key = FRawInputKeys::GenericUSBController_Button25;
+	ButtonProperties[25].Key = FRawInputKeys::GenericUSBController_Button26;
+	ButtonProperties[26].Key = FRawInputKeys::GenericUSBController_Button27;
+	ButtonProperties[27].Key = FRawInputKeys::GenericUSBController_Button28;
+	ButtonProperties[28].Key = FRawInputKeys::GenericUSBController_Button29;
+	ButtonProperties[29].Key = FRawInputKeys::GenericUSBController_Button30;
+	ButtonProperties[30].Key = FRawInputKeys::GenericUSBController_Button31;
+	ButtonProperties[31].Key = FRawInputKeys::GenericUSBController_Button32;
+	ButtonProperties[32].Key = FRawInputKeys::GenericUSBController_Button33;
+	ButtonProperties[33].Key = FRawInputKeys::GenericUSBController_Button34;
+	ButtonProperties[34].Key = FRawInputKeys::GenericUSBController_Button35;
+	ButtonProperties[35].Key = FRawInputKeys::GenericUSBController_Button36;
+	ButtonProperties[36].Key = FRawInputKeys::GenericUSBController_Button37;
+	ButtonProperties[37].Key = FRawInputKeys::GenericUSBController_Button38;
+	ButtonProperties[38].Key = FRawInputKeys::GenericUSBController_Button39;
+	ButtonProperties[39].Key = FRawInputKeys::GenericUSBController_Button40;
+	ButtonProperties[40].Key = FRawInputKeys::GenericUSBController_Button41;
+	ButtonProperties[41].Key = FRawInputKeys::GenericUSBController_Button42;
+	ButtonProperties[42].Key = FRawInputKeys::GenericUSBController_Button43;
+	ButtonProperties[43].Key = FRawInputKeys::GenericUSBController_Button44;
+	ButtonProperties[44].Key = FRawInputKeys::GenericUSBController_Button45;
+	ButtonProperties[45].Key = FRawInputKeys::GenericUSBController_Button46;
+	ButtonProperties[46].Key = FRawInputKeys::GenericUSBController_Button47;
+	ButtonProperties[47].Key = FRawInputKeys::GenericUSBController_Button48;
+	ButtonProperties[48].Key = FRawInputKeys::GenericUSBController_Button49;
+	ButtonProperties[49].Key = FRawInputKeys::GenericUSBController_Button50;
+	ButtonProperties[50].Key = FRawInputKeys::GenericUSBController_Button51;
+	ButtonProperties[51].Key = FRawInputKeys::GenericUSBController_Button52;
+	ButtonProperties[52].Key = FRawInputKeys::GenericUSBController_Button53;
+	ButtonProperties[53].Key = FRawInputKeys::GenericUSBController_Button54;
+	ButtonProperties[54].Key = FRawInputKeys::GenericUSBController_Button55;
+	ButtonProperties[55].Key = FRawInputKeys::GenericUSBController_Button56;
+	ButtonProperties[56].Key = FRawInputKeys::GenericUSBController_Button57;
+	ButtonProperties[57].Key = FRawInputKeys::GenericUSBController_Button58;
+	ButtonProperties[58].Key = FRawInputKeys::GenericUSBController_Button59;
+	ButtonProperties[59].Key = FRawInputKeys::GenericUSBController_Button60;
+	ButtonProperties[60].Key = FRawInputKeys::GenericUSBController_Button61;
+	ButtonProperties[61].Key = FRawInputKeys::GenericUSBController_Button62;
+	ButtonProperties[62].Key = FRawInputKeys::GenericUSBController_Button63;
+	ButtonProperties[63].Key = FRawInputKeys::GenericUSBController_Button64;
+	ButtonProperties[64].Key = FRawInputKeys::GenericUSBController_Button65;
+	ButtonProperties[65].Key = FRawInputKeys::GenericUSBController_Button66;
+	ButtonProperties[66].Key = FRawInputKeys::GenericUSBController_Button67;
+	ButtonProperties[67].Key = FRawInputKeys::GenericUSBController_Button68;
+	ButtonProperties[68].Key = FRawInputKeys::GenericUSBController_Button69;
+	ButtonProperties[69].Key = FRawInputKeys::GenericUSBController_Button70;
+	ButtonProperties[70].Key = FRawInputKeys::GenericUSBController_Button71;
+	ButtonProperties[71].Key = FRawInputKeys::GenericUSBController_Button72;
+	ButtonProperties[72].Key = FRawInputKeys::GenericUSBController_Button73;
+	ButtonProperties[73].Key = FRawInputKeys::GenericUSBController_Button74;
+	ButtonProperties[74].Key = FRawInputKeys::GenericUSBController_Button75;
+	ButtonProperties[75].Key = FRawInputKeys::GenericUSBController_Button76;
+	ButtonProperties[76].Key = FRawInputKeys::GenericUSBController_Button77;
+	ButtonProperties[77].Key = FRawInputKeys::GenericUSBController_Button78;
+	ButtonProperties[78].Key = FRawInputKeys::GenericUSBController_Button79;
+	ButtonProperties[79].Key = FRawInputKeys::GenericUSBController_Button80;
+	ButtonProperties[80].Key = FRawInputKeys::GenericUSBController_Button81;
+	ButtonProperties[81].Key = FRawInputKeys::GenericUSBController_Button82;
+	ButtonProperties[82].Key = FRawInputKeys::GenericUSBController_Button83;
+	ButtonProperties[83].Key = FRawInputKeys::GenericUSBController_Button84;
+	ButtonProperties[84].Key = FRawInputKeys::GenericUSBController_Button85;
+	ButtonProperties[85].Key = FRawInputKeys::GenericUSBController_Button86;
+	ButtonProperties[86].Key = FRawInputKeys::GenericUSBController_Button87;
+	ButtonProperties[87].Key = FRawInputKeys::GenericUSBController_Button88;
+	ButtonProperties[88].Key = FRawInputKeys::GenericUSBController_Button89;
+	ButtonProperties[89].Key = FRawInputKeys::GenericUSBController_Button90;
+	ButtonProperties[90].Key = FRawInputKeys::GenericUSBController_Button91;
+	ButtonProperties[91].Key = FRawInputKeys::GenericUSBController_Button92;
+	ButtonProperties[92].Key = FRawInputKeys::GenericUSBController_Button93;
+	ButtonProperties[93].Key = FRawInputKeys::GenericUSBController_Button94;
+	ButtonProperties[94].Key = FRawInputKeys::GenericUSBController_Button95;
+	ButtonProperties[95].Key = FRawInputKeys::GenericUSBController_Button96;
 
 	AxisProperties[0].Key = FRawInputKeys::GenericUSBController_Axis1;
 	AxisProperties[1].Key = FRawInputKeys::GenericUSBController_Axis2;
@@ -41,6 +117,22 @@ FRawInputDeviceConfiguration::FRawInputDeviceConfiguration()
 	AxisProperties[5].Key = FRawInputKeys::GenericUSBController_Axis6;
 	AxisProperties[6].Key = FRawInputKeys::GenericUSBController_Axis7;
 	AxisProperties[7].Key = FRawInputKeys::GenericUSBController_Axis8;
+	AxisProperties[8].Key = FRawInputKeys::GenericUSBController_Axis9;
+	AxisProperties[9].Key = FRawInputKeys::GenericUSBController_Axis10;
+	AxisProperties[10].Key = FRawInputKeys::GenericUSBController_Axis11;
+	AxisProperties[11].Key = FRawInputKeys::GenericUSBController_Axis12;
+	AxisProperties[12].Key = FRawInputKeys::GenericUSBController_Axis13;
+	AxisProperties[13].Key = FRawInputKeys::GenericUSBController_Axis14;
+	AxisProperties[14].Key = FRawInputKeys::GenericUSBController_Axis15;
+	AxisProperties[15].Key = FRawInputKeys::GenericUSBController_Axis16;
+	AxisProperties[16].Key = FRawInputKeys::GenericUSBController_Axis17;
+	AxisProperties[17].Key = FRawInputKeys::GenericUSBController_Axis18;
+	AxisProperties[18].Key = FRawInputKeys::GenericUSBController_Axis19;
+	AxisProperties[19].Key = FRawInputKeys::GenericUSBController_Axis20;
+	AxisProperties[20].Key = FRawInputKeys::GenericUSBController_Axis21;
+	AxisProperties[21].Key = FRawInputKeys::GenericUSBController_Axis22;
+	AxisProperties[22].Key = FRawInputKeys::GenericUSBController_Axis23;
+	AxisProperties[23].Key = FRawInputKeys::GenericUSBController_Axis24;
 }
 
 FName URawInputSettings::GetCategoryName() const
diff --git a/Engine/Plugins/Experimental/RawInput/Source/RawInput/Private/Windows/RawInputWindows.cpp b/Engine/Plugins/Experimental/RawInput/Source/RawInput/Private/Windows/RawInputWindows.cpp
--- a/Engine/Plugins/Experimental/RawInput/Source/RawInput/Private/Windows/RawInputWindows.cpp
+++ b/Engine/Plugins/Experimental/RawInput/Source/RawInput/Private/Windows/RawInputWindows.cpp
@@ -100,12 +100,13 @@ FRawInputWindows::FRawInputWindows(const TSharedRef<FGenericApplicationMessageHa
 		const uint32 Flags = 0;
 		const int32 PageID = 0x01;
 		int32 DeviceID = 0x04;
-		DefaultDeviceHandle = RegisterInputDevice(RIM_TYPEHID, Flags, DeviceID, PageID);
+		
+		DefaultDeviceHandle = RegisterInputDevice(RIM_TYPEHID, Flags, DeviceID, PageID, nullptr);
 
 		if (DefaultDeviceHandle == INDEX_NONE)
 		{
 			DeviceID = 0x05;
-			DefaultDeviceHandle = RegisterInputDevice(RIM_TYPEHID, Flags, DeviceID, PageID);
+			DefaultDeviceHandle = RegisterInputDevice(RIM_TYPEHID, Flags, DeviceID, PageID, nullptr);
 		}
 	}
 
@@ -141,10 +142,10 @@ HWND GetWindowHandle()
 	return hWnd;
 }
 
-int32 FRawInputWindows::RegisterInputDevice(const int32 DeviceType, const int32 Flags, const uint16 DeviceID, const int16 PageID)
+int32 FRawInputWindows::RegisterInputDevice(const int32 DeviceType, const int32 Flags, const uint16 DeviceID, const int16 PageID, const HANDLE Handle)
 {
 	int32 DeviceHandle = INDEX_NONE;
-	FRawInputRegisteredDevice DeviceData(DeviceType, DeviceID, PageID);
+	FRawInputRegisteredDevice DeviceData(DeviceType, DeviceID, PageID, Handle);
 	
 	RAWINPUTDEVICE RawInputDevice;
 
@@ -170,33 +171,27 @@ int32 FRawInputWindows::RegisterInputDevice(const int32 DeviceType, const int32
 		DeviceHandle = FindRegisteredDeviceHandle(DeviceData);
 		if (DeviceHandle == INDEX_NONE)
 		{
-			DeviceHandle = GetNextInputHandle();
-			RegisteredDeviceList.Add(DeviceHandle, DeviceData);
-			// Now see if the device is connected
+			// Register all devices of the type specified in the param
+			// NOTE: This it kind of a hack. A true refactor of the RawInputPlugin is required but this will make it usable with multiple complementary devices. 
 			bool bWasConnected = false;
 			for (const FConnectedDeviceInfo& ConnectedDeviceInfo : ConnectedDeviceInfoList)
 			{
 				if (CompareDeviceInfo(ConnectedDeviceInfo.RIDDeviceInfo, DeviceData))
 				{
+					DeviceHandle = GetNextInputHandle();
+					RegisteredDeviceList.Add(DeviceHandle, DeviceData);
+
 					FRawWindowsDeviceEntry& RegisteredDeviceInfo = RegisteredDeviceList[DeviceHandle];
 					CopyConnectedDeviceInfo(RegisteredDeviceInfo, &ConnectedDeviceInfo);
 
 					UE_LOG(LogRawInputWindows, Log, TEXT("VendorID:%04X ProductID:%04X"), RegisteredDeviceInfo.DeviceData.VendorID, RegisteredDeviceInfo.DeviceData.ProductID);
 
 					bWasConnected = true;
-					break;
-				}
-			}
-			if (bWasConnected)
-			{
-				SetupBindings(DeviceHandle, true);
 
-				UE_LOG(LogRawInputWindows, Log, TEXT("Device was registered succesfully and is connected (Usage:%d UsagePage:%d)"), DeviceData.Usage, DeviceData.UsagePage);				
-			}
-			else
-			{
-				DeviceHandle = INDEX_NONE;
-				UE_LOG(LogRawInputWindows, Warning, TEXT("Device was registered succesfully but not connected (Usage:%d UsagePage:%d)"), DeviceData.Usage, DeviceData.UsagePage);
+					SetupBindings(DeviceHandle, true);
+
+					UE_LOG(LogRawInputWindows, Log, TEXT("Device was registered succesfully and is connected (Usage:%d UsagePage:%d)"), DeviceData.Usage, DeviceData.UsagePage);
+				}
 			}
 		}
 		else
@@ -290,6 +285,22 @@ void FRawInputWindows::SetupBindings(const int32 DeviceHandle, const bool bApply
 		BindAnalogForDevice(DeviceHandle, FRawInputKeyNames::GenericUSBController_Axis6, 5);
 		BindAnalogForDevice(DeviceHandle, FRawInputKeyNames::GenericUSBController_Axis7, 6);
 		BindAnalogForDevice(DeviceHandle, FRawInputKeyNames::GenericUSBController_Axis8, 7);
+		BindAnalogForDevice(DeviceHandle, FRawInputKeyNames::GenericUSBController_Axis9, 8);
+		BindAnalogForDevice(DeviceHandle, FRawInputKeyNames::GenericUSBController_Axis10, 9);
+		BindAnalogForDevice(DeviceHandle, FRawInputKeyNames::GenericUSBController_Axis11, 10);
+		BindAnalogForDevice(DeviceHandle, FRawInputKeyNames::GenericUSBController_Axis12, 11);
+		BindAnalogForDevice(DeviceHandle, FRawInputKeyNames::GenericUSBController_Axis13, 12);
+		BindAnalogForDevice(DeviceHandle, FRawInputKeyNames::GenericUSBController_Axis14, 13);
+		BindAnalogForDevice(DeviceHandle, FRawInputKeyNames::GenericUSBController_Axis15, 14);
+		BindAnalogForDevice(DeviceHandle, FRawInputKeyNames::GenericUSBController_Axis16, 15);
+		BindAnalogForDevice(DeviceHandle, FRawInputKeyNames::GenericUSBController_Axis17, 16);
+		BindAnalogForDevice(DeviceHandle, FRawInputKeyNames::GenericUSBController_Axis18, 17);
+		BindAnalogForDevice(DeviceHandle, FRawInputKeyNames::GenericUSBController_Axis19, 18);
+		BindAnalogForDevice(DeviceHandle, FRawInputKeyNames::GenericUSBController_Axis20, 19);
+		BindAnalogForDevice(DeviceHandle, FRawInputKeyNames::GenericUSBController_Axis21, 20);
+		BindAnalogForDevice(DeviceHandle, FRawInputKeyNames::GenericUSBController_Axis22, 21);
+		BindAnalogForDevice(DeviceHandle, FRawInputKeyNames::GenericUSBController_Axis23, 22);
+		BindAnalogForDevice(DeviceHandle, FRawInputKeyNames::GenericUSBController_Axis24, 23);
 
 		BindButtonForDevice(DeviceHandle, FRawInputKeyNames::GenericUSBController_Button1, 0);
 		BindButtonForDevice(DeviceHandle, FRawInputKeyNames::GenericUSBController_Button2, 1);
@@ -311,6 +322,82 @@ void FRawInputWindows::SetupBindings(const int32 DeviceHandle, const bool bApply
 		BindButtonForDevice(DeviceHandle, FRawInputKeyNames::GenericUSBController_Button18, 17);
 		BindButtonForDevice(DeviceHandle, FRawInputKeyNames::GenericUSBController_Button19, 18);
 		BindButtonForDevice(DeviceHandle, FRawInputKeyNames::GenericUSBController_Button20, 19);
+		BindButtonForDevice(DeviceHandle, FRawInputKeyNames::GenericUSBController_Button21, 20);
+		BindButtonForDevice(DeviceHandle, FRawInputKeyNames::GenericUSBController_Button22, 21);
+		BindButtonForDevice(DeviceHandle, FRawInputKeyNames::GenericUSBController_Button23, 22);
+		BindButtonForDevice(DeviceHandle, FRawInputKeyNames::GenericUSBController_Button24, 23);
+		BindButtonForDevice(DeviceHandle, FRawInputKeyNames::GenericUSBController_Button25, 24);
+		BindButtonForDevice(DeviceHandle, FRawInputKeyNames::GenericUSBController_Button26, 25);
+		BindButtonForDevice(DeviceHandle, FRawInputKeyNames::GenericUSBController_Button27, 26);
+		BindButtonForDevice(DeviceHandle, FRawInputKeyNames::GenericUSBController_Button28, 27);
+		BindButtonForDevice(DeviceHandle, FRawInputKeyNames::GenericUSBController_Button29, 28);
+		BindButtonForDevice(DeviceHandle, FRawInputKeyNames::GenericUSBController_Button30, 29);
+		BindButtonForDevice(DeviceHandle, FRawInputKeyNames::GenericUSBController_Button31, 30);
+		BindButtonForDevice(DeviceHandle, FRawInputKeyNames::GenericUSBController_Button32, 31);
+		BindButtonForDevice(DeviceHandle, FRawInputKeyNames::GenericUSBController_Button33, 32);
+		BindButtonForDevice(DeviceHandle, FRawInputKeyNames::GenericUSBController_Button34, 33);
+		BindButtonForDevice(DeviceHandle, FRawInputKeyNames::GenericUSBController_Button35, 34);
+		BindButtonForDevice(DeviceHandle, FRawInputKeyNames::GenericUSBController_Button36, 35);
+		BindButtonForDevice(DeviceHandle, FRawInputKeyNames::GenericUSBController_Button37, 36);
+		BindButtonForDevice(DeviceHandle, FRawInputKeyNames::GenericUSBController_Button38, 37);
+		BindButtonForDevice(DeviceHandle, FRawInputKeyNames::GenericUSBController_Button39, 38);
+		BindButtonForDevice(DeviceHandle, FRawInputKeyNames::GenericUSBController_Button40, 39);
+		BindButtonForDevice(DeviceHandle, FRawInputKeyNames::GenericUSBController_Button41, 40);
+		BindButtonForDevice(DeviceHandle, FRawInputKeyNames::GenericUSBController_Button42, 41);
+		BindButtonForDevice(DeviceHandle, FRawInputKeyNames::GenericUSBController_Button43, 42);
+		BindButtonForDevice(DeviceHandle, FRawInputKeyNames::GenericUSBController_Button44, 43);
+		BindButtonForDevice(DeviceHandle, FRawInputKeyNames::GenericUSBController_Button45, 44);
+		BindButtonForDevice(DeviceHandle, FRawInputKeyNames::GenericUSBController_Button46, 45);
+		BindButtonForDevice(DeviceHandle, FRawInputKeyNames::GenericUSBController_Button47, 46);
+		BindButtonForDevice(DeviceHandle, FRawInputKeyNames::GenericUSBController_Button48, 47);
+		BindButtonForDevice(DeviceHandle, FRawInputKeyNames::GenericUSBController_Button49, 48);
+		BindButtonForDevice(DeviceHandle, FRawInputKeyNames::GenericUSBController_Button50, 49);
+		BindButtonForDevice(DeviceHandle, FRawInputKeyNames::GenericUSBController_Button51, 50);
+		BindButtonForDevice(DeviceHandle, FRawInputKeyNames::GenericUSBController_Button52, 51);
+		BindButtonForDevice(DeviceHandle, FRawInputKeyNames::GenericUSBController_Button53, 52);
+		BindButtonForDevice(DeviceHandle, FRawInputKeyNames::GenericUSBController_Button54, 53);
+		BindButtonForDevice(DeviceHandle, FRawInputKeyNames::GenericUSBController_Button55, 54);
+		BindButtonForDevice(DeviceHandle, FRawInputKeyNames::GenericUSBController_Button56, 55);
+		BindButtonForDevice(DeviceHandle, FRawInputKeyNames::GenericUSBController_Button57, 56);
+		BindButtonForDevice(DeviceHandle, FRawInputKeyNames::GenericUSBController_Button58, 57);
+		BindButtonForDevice(DeviceHandle, FRawInputKeyNames::GenericUSBController_Button59, 58);
+		BindButtonForDevice(DeviceHandle, FRawInputKeyNames::GenericUSBController_Button60, 59);
+		BindButtonForDevice(DeviceHandle, FRawInputKeyNames::GenericUSBController_Button61, 60);
+		BindButtonForDevice(DeviceHandle, FRawInputKeyNames::GenericUSBController_Button62, 61);
+		BindButtonForDevice(DeviceHandle, FRawInputKeyNames::GenericUSBController_Button63, 62);
+		BindButtonForDevice(DeviceHandle, FRawInputKeyNames::GenericUSBController_Button64, 63);
+		BindButtonForDevice(DeviceHandle, FRawInputKeyNames::GenericUSBController_Button65, 64);
+		BindButtonForDevice(DeviceHandle, FRawInputKeyNames::GenericUSBController_Button66, 65);
+		BindButtonForDevice(DeviceHandle, FRawInputKeyNames::GenericUSBController_Button67, 66);
+		BindButtonForDevice(DeviceHandle, FRawInputKeyNames::GenericUSBController_Button68, 67);
+		BindButtonForDevice(DeviceHandle, FRawInputKeyNames::GenericUSBController_Button69, 68);
+		BindButtonForDevice(DeviceHandle, FRawInputKeyNames::GenericUSBController_Button70, 69);
+		BindButtonForDevice(DeviceHandle, FRawInputKeyNames::GenericUSBController_Button71, 70);
+		BindButtonForDevice(DeviceHandle, FRawInputKeyNames::GenericUSBController_Button72, 71);
+		BindButtonForDevice(DeviceHandle, FRawInputKeyNames::GenericUSBController_Button73, 72);
+		BindButtonForDevice(DeviceHandle, FRawInputKeyNames::GenericUSBController_Button74, 73);
+		BindButtonForDevice(DeviceHandle, FRawInputKeyNames::GenericUSBController_Button75, 74);
+		BindButtonForDevice(DeviceHandle, FRawInputKeyNames::GenericUSBController_Button76, 75);
+		BindButtonForDevice(DeviceHandle, FRawInputKeyNames::GenericUSBController_Button77, 76);
+		BindButtonForDevice(DeviceHandle, FRawInputKeyNames::GenericUSBController_Button78, 77);
+		BindButtonForDevice(DeviceHandle, FRawInputKeyNames::GenericUSBController_Button79, 78);
+		BindButtonForDevice(DeviceHandle, FRawInputKeyNames::GenericUSBController_Button80, 79);
+		BindButtonForDevice(DeviceHandle, FRawInputKeyNames::GenericUSBController_Button81, 80);
+		BindButtonForDevice(DeviceHandle, FRawInputKeyNames::GenericUSBController_Button82, 81);
+		BindButtonForDevice(DeviceHandle, FRawInputKeyNames::GenericUSBController_Button83, 82);
+		BindButtonForDevice(DeviceHandle, FRawInputKeyNames::GenericUSBController_Button84, 83);
+		BindButtonForDevice(DeviceHandle, FRawInputKeyNames::GenericUSBController_Button85, 84);
+		BindButtonForDevice(DeviceHandle, FRawInputKeyNames::GenericUSBController_Button86, 85);
+		BindButtonForDevice(DeviceHandle, FRawInputKeyNames::GenericUSBController_Button87, 86);
+		BindButtonForDevice(DeviceHandle, FRawInputKeyNames::GenericUSBController_Button88, 87);
+		BindButtonForDevice(DeviceHandle, FRawInputKeyNames::GenericUSBController_Button89, 88);
+		BindButtonForDevice(DeviceHandle, FRawInputKeyNames::GenericUSBController_Button90, 89);
+		BindButtonForDevice(DeviceHandle, FRawInputKeyNames::GenericUSBController_Button91, 90);
+		BindButtonForDevice(DeviceHandle, FRawInputKeyNames::GenericUSBController_Button92, 91);
+		BindButtonForDevice(DeviceHandle, FRawInputKeyNames::GenericUSBController_Button93, 92);
+		BindButtonForDevice(DeviceHandle, FRawInputKeyNames::GenericUSBController_Button94, 93);
+		BindButtonForDevice(DeviceHandle, FRawInputKeyNames::GenericUSBController_Button95, 94);
+		BindButtonForDevice(DeviceHandle, FRawInputKeyNames::GenericUSBController_Button96, 95);
 	}
 }
 
@@ -461,7 +548,7 @@ bool FRawInputWindows::ProcessMessage(const HWND hwnd, const uint32 Msg, const W
 						// now that we have the PP data we need to get the caps, check those and see if this is a device we registered and if it is store it so we can send it
 						if (DLLPointers.HidP_GetCaps((PHIDP_PREPARSED_DATA)PreParsedData.GetData(), &Caps) == HIDP_STATUS_SUCCESS)
 						{
-							DeviceData = FRawInputRegisteredDevice(RawInputDataBuffer->header.dwType, Caps.Usage, Caps.UsagePage);									
+							DeviceData = FRawInputRegisteredDevice(RawInputDataBuffer->header.dwType, Caps.Usage, Caps.UsagePage, RawInputDataBuffer->header.hDevice);
 						}
 					}							
 				}
@@ -723,7 +810,8 @@ void FRawInputWindows::QueryConnectedDevices()
 		}
 
 		// Add to the list of registered devices
-		ConnectedDeviceInfoList.Add(FConnectedDeviceInfo(MoveTemp(DeviceName), RawDeviceInfo));
+		int32 index = ConnectedDeviceInfoList.Add(FConnectedDeviceInfo(MoveTemp(DeviceName), RawDeviceInfo, Device.hDevice));
+
 		ShowDeviceInfo(ConnectedDeviceInfoList.Last());
 	}
 
@@ -859,6 +947,8 @@ void FRawInputWindows::CopyConnectedDeviceInfo(FRawWindowsDeviceEntry& Registere
 			RegisteredDevice.DeviceData.VendorID = ConnectedDevice->RIDDeviceInfo.hid.dwVendorId;
 			RegisteredDevice.DeviceData.ProductID = ConnectedDevice->RIDDeviceInfo.hid.dwProductId;
 		}
+
+		RegisteredDevice.DeviceData.hDevice = ConnectedDevice->hDevice;
 	}
 	else
 	{
diff --git a/Engine/Plugins/Experimental/RawInput/Source/RawInput/Public/RawInput.h b/Engine/Plugins/Experimental/RawInput/Source/RawInput/Public/RawInput.h
--- a/Engine/Plugins/Experimental/RawInput/Source/RawInput/Public/RawInput.h
+++ b/Engine/Plugins/Experimental/RawInput/Source/RawInput/Public/RawInput.h
@@ -21,7 +21,7 @@ public:
 	virtual void QueryConnectedDevices() {}
 
 	/** Register a device for use. */
-	virtual int32 RegisterInputDevice(int32 DeviceType, int32 Flags, uint16 DeviceID, int16 PageID) = 0;
+	virtual int32 RegisterInputDevice(int32 DeviceType, int32 Flags, uint16 DeviceID, int16 PageID, void* Handle) = 0;
 	
 	/** Remove a registered device */
 	virtual void RemoveRegisteredInputDevice(int32 DeviceHandle) {};
diff --git a/Engine/Plugins/Experimental/RawInput/Source/RawInput/Public/RawInputFunctionLibrary.h b/Engine/Plugins/Experimental/RawInput/Source/RawInput/Public/RawInputFunctionLibrary.h
--- a/Engine/Plugins/Experimental/RawInput/Source/RawInput/Public/RawInputFunctionLibrary.h
+++ b/Engine/Plugins/Experimental/RawInput/Source/RawInput/Public/RawInputFunctionLibrary.h
@@ -18,6 +18,22 @@ struct RAWINPUT_API FRawInputKeyNames
 	static const FGamepadKeyNames::Type GenericUSBController_Axis6;
 	static const FGamepadKeyNames::Type GenericUSBController_Axis7;
 	static const FGamepadKeyNames::Type GenericUSBController_Axis8;
+	static const FGamepadKeyNames::Type GenericUSBController_Axis9;
+	static const FGamepadKeyNames::Type GenericUSBController_Axis10;
+	static const FGamepadKeyNames::Type GenericUSBController_Axis11;
+	static const FGamepadKeyNames::Type GenericUSBController_Axis12;
+	static const FGamepadKeyNames::Type GenericUSBController_Axis13;
+	static const FGamepadKeyNames::Type GenericUSBController_Axis14;
+	static const FGamepadKeyNames::Type GenericUSBController_Axis15;
+	static const FGamepadKeyNames::Type GenericUSBController_Axis16;
+	static const FGamepadKeyNames::Type GenericUSBController_Axis17;
+	static const FGamepadKeyNames::Type GenericUSBController_Axis18;
+	static const FGamepadKeyNames::Type GenericUSBController_Axis19;
+	static const FGamepadKeyNames::Type GenericUSBController_Axis20;
+	static const FGamepadKeyNames::Type GenericUSBController_Axis21;
+	static const FGamepadKeyNames::Type GenericUSBController_Axis22;
+	static const FGamepadKeyNames::Type GenericUSBController_Axis23;
+	static const FGamepadKeyNames::Type GenericUSBController_Axis24;
 
 	static const FGamepadKeyNames::Type GenericUSBController_Button1;
 	static const FGamepadKeyNames::Type GenericUSBController_Button2;
@@ -39,6 +55,82 @@ struct RAWINPUT_API FRawInputKeyNames
 	static const FGamepadKeyNames::Type GenericUSBController_Button18;
 	static const FGamepadKeyNames::Type GenericUSBController_Button19;
 	static const FGamepadKeyNames::Type GenericUSBController_Button20;
+	static const FGamepadKeyNames::Type GenericUSBController_Button21;
+	static const FGamepadKeyNames::Type GenericUSBController_Button22;
+	static const FGamepadKeyNames::Type GenericUSBController_Button23;
+	static const FGamepadKeyNames::Type GenericUSBController_Button24;
+	static const FGamepadKeyNames::Type GenericUSBController_Button25;
+	static const FGamepadKeyNames::Type GenericUSBController_Button26;
+	static const FGamepadKeyNames::Type GenericUSBController_Button27;
+	static const FGamepadKeyNames::Type GenericUSBController_Button28;
+	static const FGamepadKeyNames::Type GenericUSBController_Button29;
+	static const FGamepadKeyNames::Type GenericUSBController_Button30;
+	static const FGamepadKeyNames::Type GenericUSBController_Button31;
+	static const FGamepadKeyNames::Type GenericUSBController_Button32;
+	static const FGamepadKeyNames::Type GenericUSBController_Button33;
+	static const FGamepadKeyNames::Type GenericUSBController_Button34;
+	static const FGamepadKeyNames::Type GenericUSBController_Button35;
+	static const FGamepadKeyNames::Type GenericUSBController_Button36;
+	static const FGamepadKeyNames::Type GenericUSBController_Button37;
+	static const FGamepadKeyNames::Type GenericUSBController_Button38;
+	static const FGamepadKeyNames::Type GenericUSBController_Button39;
+	static const FGamepadKeyNames::Type GenericUSBController_Button40;
+	static const FGamepadKeyNames::Type GenericUSBController_Button41;
+	static const FGamepadKeyNames::Type GenericUSBController_Button42;
+	static const FGamepadKeyNames::Type GenericUSBController_Button43;
+	static const FGamepadKeyNames::Type GenericUSBController_Button44;
+	static const FGamepadKeyNames::Type GenericUSBController_Button45;
+	static const FGamepadKeyNames::Type GenericUSBController_Button46;
+	static const FGamepadKeyNames::Type GenericUSBController_Button47;
+	static const FGamepadKeyNames::Type GenericUSBController_Button48;
+	static const FGamepadKeyNames::Type GenericUSBController_Button49;
+	static const FGamepadKeyNames::Type GenericUSBController_Button50;
+	static const FGamepadKeyNames::Type GenericUSBController_Button51;
+	static const FGamepadKeyNames::Type GenericUSBController_Button52;
+	static const FGamepadKeyNames::Type GenericUSBController_Button53;
+	static const FGamepadKeyNames::Type GenericUSBController_Button54;
+	static const FGamepadKeyNames::Type GenericUSBController_Button55;
+	static const FGamepadKeyNames::Type GenericUSBController_Button56;
+	static const FGamepadKeyNames::Type GenericUSBController_Button57;
+	static const FGamepadKeyNames::Type GenericUSBController_Button58;
+	static const FGamepadKeyNames::Type GenericUSBController_Button59;
+	static const FGamepadKeyNames::Type GenericUSBController_Button60;
+	static const FGamepadKeyNames::Type GenericUSBController_Button61;
+	static const FGamepadKeyNames::Type GenericUSBController_Button62;
+	static const FGamepadKeyNames::Type GenericUSBController_Button63;
+	static const FGamepadKeyNames::Type GenericUSBController_Button64;
+	static const FGamepadKeyNames::Type GenericUSBController_Button65;
+	static const FGamepadKeyNames::Type GenericUSBController_Button66;
+	static const FGamepadKeyNames::Type GenericUSBController_Button67;
+	static const FGamepadKeyNames::Type GenericUSBController_Button68;
+	static const FGamepadKeyNames::Type GenericUSBController_Button69;
+	static const FGamepadKeyNames::Type GenericUSBController_Button70;
+	static const FGamepadKeyNames::Type GenericUSBController_Button71;
+	static const FGamepadKeyNames::Type GenericUSBController_Button72;
+	static const FGamepadKeyNames::Type GenericUSBController_Button73;
+	static const FGamepadKeyNames::Type GenericUSBController_Button74;
+	static const FGamepadKeyNames::Type GenericUSBController_Button75;
+	static const FGamepadKeyNames::Type GenericUSBController_Button76;
+	static const FGamepadKeyNames::Type GenericUSBController_Button77;
+	static const FGamepadKeyNames::Type GenericUSBController_Button78;
+	static const FGamepadKeyNames::Type GenericUSBController_Button79;
+	static const FGamepadKeyNames::Type GenericUSBController_Button80;
+	static const FGamepadKeyNames::Type GenericUSBController_Button81;
+	static const FGamepadKeyNames::Type GenericUSBController_Button82;
+	static const FGamepadKeyNames::Type GenericUSBController_Button83;
+	static const FGamepadKeyNames::Type GenericUSBController_Button84;
+	static const FGamepadKeyNames::Type GenericUSBController_Button85;
+	static const FGamepadKeyNames::Type GenericUSBController_Button86;
+	static const FGamepadKeyNames::Type GenericUSBController_Button87;
+	static const FGamepadKeyNames::Type GenericUSBController_Button88;
+	static const FGamepadKeyNames::Type GenericUSBController_Button89;
+	static const FGamepadKeyNames::Type GenericUSBController_Button90;
+	static const FGamepadKeyNames::Type GenericUSBController_Button91;
+	static const FGamepadKeyNames::Type GenericUSBController_Button92;
+	static const FGamepadKeyNames::Type GenericUSBController_Button93;
+	static const FGamepadKeyNames::Type GenericUSBController_Button94;
+	static const FGamepadKeyNames::Type GenericUSBController_Button95;
+	static const FGamepadKeyNames::Type GenericUSBController_Button96;
 };
 
 struct RAWINPUT_API FRawInputKeys
@@ -51,6 +143,22 @@ struct RAWINPUT_API FRawInputKeys
 	static const FKey GenericUSBController_Axis6;
 	static const FKey GenericUSBController_Axis7;
 	static const FKey GenericUSBController_Axis8;
+	static const FKey GenericUSBController_Axis9;
+	static const FKey GenericUSBController_Axis10;
+	static const FKey GenericUSBController_Axis11;
+	static const FKey GenericUSBController_Axis12;
+	static const FKey GenericUSBController_Axis13;
+	static const FKey GenericUSBController_Axis14;
+	static const FKey GenericUSBController_Axis15;
+	static const FKey GenericUSBController_Axis16;
+	static const FKey GenericUSBController_Axis17;
+	static const FKey GenericUSBController_Axis18;
+	static const FKey GenericUSBController_Axis19;
+	static const FKey GenericUSBController_Axis20;
+	static const FKey GenericUSBController_Axis21;
+	static const FKey GenericUSBController_Axis22;
+	static const FKey GenericUSBController_Axis23;
+	static const FKey GenericUSBController_Axis24;
 
 	static const FKey GenericUSBController_Button1;
 	static const FKey GenericUSBController_Button2;
@@ -72,6 +180,83 @@ struct RAWINPUT_API FRawInputKeys
 	static const FKey GenericUSBController_Button18;
 	static const FKey GenericUSBController_Button19;
 	static const FKey GenericUSBController_Button20;
+	static const FKey GenericUSBController_Button21;
+	static const FKey GenericUSBController_Button22;
+	static const FKey GenericUSBController_Button23;
+	static const FKey GenericUSBController_Button24;
+	static const FKey GenericUSBController_Button25;
+	static const FKey GenericUSBController_Button26;
+	static const FKey GenericUSBController_Button27;
+	static const FKey GenericUSBController_Button28;
+	static const FKey GenericUSBController_Button29;
+	static const FKey GenericUSBController_Button30;
+	static const FKey GenericUSBController_Button31;
+	static const FKey GenericUSBController_Button32;
+	static const FKey GenericUSBController_Button33;
+	static const FKey GenericUSBController_Button34;
+	static const FKey GenericUSBController_Button35;
+	static const FKey GenericUSBController_Button36;
+	static const FKey GenericUSBController_Button37;
+	static const FKey GenericUSBController_Button38;
+	static const FKey GenericUSBController_Button39;
+	static const FKey GenericUSBController_Button40;
+	static const FKey GenericUSBController_Button41;
+	static const FKey GenericUSBController_Button42;
+	static const FKey GenericUSBController_Button43;
+	static const FKey GenericUSBController_Button44;
+	static const FKey GenericUSBController_Button45;
+	static const FKey GenericUSBController_Button46;
+	static const FKey GenericUSBController_Button47;
+	static const FKey GenericUSBController_Button48;
+	static const FKey GenericUSBController_Button49;
+	static const FKey GenericUSBController_Button50;
+	static const FKey GenericUSBController_Button51;
+	static const FKey GenericUSBController_Button52;
+	static const FKey GenericUSBController_Button53;
+	static const FKey GenericUSBController_Button54;
+	static const FKey GenericUSBController_Button55;
+	static const FKey GenericUSBController_Button56;
+	static const FKey GenericUSBController_Button57;
+	static const FKey GenericUSBController_Button58;
+	static const FKey GenericUSBController_Button59;
+	static const FKey GenericUSBController_Button60;
+	static const FKey GenericUSBController_Button61;
+	static const FKey GenericUSBController_Button62;
+	static const FKey GenericUSBController_Button63;
+	static const FKey GenericUSBController_Button64;
+	static const FKey GenericUSBController_Button65;
+	static const FKey GenericUSBController_Button66;
+	static const FKey GenericUSBController_Button67;
+	static const FKey GenericUSBController_Button68;
+	static const FKey GenericUSBController_Button69;
+	static const FKey GenericUSBController_Button70;
+	static const FKey GenericUSBController_Button71;
+	static const FKey GenericUSBController_Button72;
+	static const FKey GenericUSBController_Button73;
+	static const FKey GenericUSBController_Button74;
+	static const FKey GenericUSBController_Button75;
+	static const FKey GenericUSBController_Button76;
+	static const FKey GenericUSBController_Button77;
+	static const FKey GenericUSBController_Button78;
+	static const FKey GenericUSBController_Button79;
+	static const FKey GenericUSBController_Button80;
+	static const FKey GenericUSBController_Button81;
+	static const FKey GenericUSBController_Button82;
+	static const FKey GenericUSBController_Button83;
+	static const FKey GenericUSBController_Button84;
+	static const FKey GenericUSBController_Button85;
+	static const FKey GenericUSBController_Button86;
+	static const FKey GenericUSBController_Button87;
+	static const FKey GenericUSBController_Button88;
+	static const FKey GenericUSBController_Button89;
+	static const FKey GenericUSBController_Button90;
+	static const FKey GenericUSBController_Button91;
+	static const FKey GenericUSBController_Button92;
+	static const FKey GenericUSBController_Button93;
+	static const FKey GenericUSBController_Button94;
+	static const FKey GenericUSBController_Button95;
+	static const FKey GenericUSBController_Button96;
+
 };
 
 USTRUCT(BlueprintType)
diff --git a/Engine/Plugins/Experimental/RawInput/Source/RawInput/Public/Windows/RawInputWindows.h b/Engine/Plugins/Experimental/RawInput/Source/RawInput/Public/Windows/RawInputWindows.h
--- a/Engine/Plugins/Experimental/RawInput/Source/RawInput/Public/Windows/RawInputWindows.h
+++ b/Engine/Plugins/Experimental/RawInput/Source/RawInput/Public/Windows/RawInputWindows.h
@@ -11,8 +11,8 @@
 
 #include "hidsdi.h"
 
-#define MAX_NUM_CONTROLLER_BUTTONS 20
-#define MAX_NUM_CONTROLLER_ANALOG 8
+#define MAX_NUM_CONTROLLER_BUTTONS 96
+#define MAX_NUM_CONTROLLER_ANALOG 24
 #define RAW_INPUT_ERROR (uint32)(-1)
 
 class FRawInputWindows;
@@ -62,8 +62,9 @@ struct FRawInputRegisteredDevice
 	{
 	}
 
-	FRawInputRegisteredDevice(const int32 InDeviceType, const uint16 InUsage, const int16 InUsagePage)
-		: VendorID(0)
+	FRawInputRegisteredDevice(const int32 InDeviceType, const uint16 InUsage, const int16 InUsagePage, const HANDLE _handle)
+		: hDevice(_handle)
+		, VendorID(0)
 		, ProductID(0)
 		, DeviceType(InDeviceType)
 		, Usage(InUsage)
@@ -76,13 +77,16 @@ struct FRawInputRegisteredDevice
 	{
 		return (bIsValid && 
 			 Other.bIsValid &&
+			 hDevice == Other.hDevice);
+			 /*
 			 (DeviceType == Other.DeviceType) &&
 			 (Usage == Other.Usage) && 
-			 (UsagePage == Other.UsagePage));
+			 (UsagePage == Other.UsagePage));*/
 	};
-
+	
 	// Driver supplied device name
 	FString DeviceName;
+	HANDLE hDevice;
 
 	// Device identifiers
 	int32 VendorID;
@@ -218,14 +222,16 @@ struct FConnectedDeviceInfo
 {
 public:
 	FConnectedDeviceInfo() {}
-	FConnectedDeviceInfo(FString InDeviceName, const RID_DEVICE_INFO& InRIDDeviceInfo)
+	FConnectedDeviceInfo(FString InDeviceName, const RID_DEVICE_INFO& InRIDDeviceInfo, const HANDLE Handle)
 		: DeviceName(MoveTemp(InDeviceName))
 		, RIDDeviceInfo(InRIDDeviceInfo)
+		, hDevice(Handle)
 	{
 	}
 
 	FString DeviceName;
 	RID_DEVICE_INFO RIDDeviceInfo;
+	HANDLE hDevice;
 };
 
 class FRawInputWindows : public IRawInput, IWindowsMessageHandler
@@ -238,7 +244,7 @@ public:
 
 	// Begin RawInput interface
 	virtual void QueryConnectedDevices() override;	
-	virtual int32 RegisterInputDevice(int32 DeviceType, int32 Flags, uint16 DeviceID, int16 PageID);
+	virtual int32 RegisterInputDevice(int32 DeviceType, int32 Flags, uint16 DeviceID, int16 PageID, HANDLE Handle) override;
 	virtual void RemoveRegisteredInputDevice(int32 DeviceHandle) override;	
 	virtual void BindButtonForDevice(int32 DeviceHandle, FName EventName, int32 ButtonIndex ) override;
 	virtual void BindAnalogForDevice(int32 DeviceHandle, FName EventName, int32 AxisIndex ) override;
