Update RawInput to use the new OnController functions instead of the soon to be deprecated old version

#rnx

diff --git a/Engine/Plugins/Experimental/RawInput/Source/RawInput/Private/Windows/RawInputWindows.cpp b/Engine/Plugins/Experimental/RawInput/Source/RawInput/Private/Windows/RawInputWindows.cpp
--- a/Engine/Plugins/Experimental/RawInput/Source/RawInput/Private/Windows/RawInputWindows.cpp
+++ b/Engine/Plugins/Experimental/RawInput/Source/RawInput/Private/Windows/RawInputWindows.cpp
@@ -993,6 +993,9 @@ static FName RawInputInterfaceName = FName("RawInput");
 
 void FRawInputWindows::SendControllerEvents()
 {
+	FPlatformUserId UserId = IPlatformInputDeviceMapper::Get().GetPrimaryPlatformUser();
+	FInputDeviceId DeviceId = IPlatformInputDeviceMapper::Get().GetDefaultInputDevice();
+	
 	for (TPair<int32, FRawWindowsDeviceEntry>& DeviceEntryPair : RegisteredDeviceList)
 	{
 		FRawWindowsDeviceEntry& DeviceEntry = DeviceEntryPair.Value;
@@ -1011,17 +1014,17 @@ void FRawInputWindows::SendControllerEvents()
 					{
 						if (DeviceButtonData.bButtonState)
 						{
-							MessageHandler->OnControllerButtonPressed(DeviceButtonData.ButtonName, 0, false);								
+							MessageHandler->OnControllerButtonPressed(DeviceButtonData.ButtonName, UserId, DeviceId, false);								
 						}
 						else
 						{
-							MessageHandler->OnControllerButtonReleased(DeviceButtonData.ButtonName, 0, false);
+							MessageHandler->OnControllerButtonReleased(DeviceButtonData.ButtonName, UserId, DeviceId, false);
 						}
 						DeviceButtonData.bPreviousButtonState = DeviceButtonData.bButtonState;
 					}
 					else if (DeviceButtonData.bButtonState) // state not changed - but is true.. which means it must have been true last time.. so is repeat
 					{
-						MessageHandler->OnControllerButtonPressed(DeviceButtonData.ButtonName, 0, true);							
+						MessageHandler->OnControllerButtonPressed(DeviceButtonData.ButtonName, UserId, DeviceId, true);							
 					}
 						
 				}
@@ -1031,7 +1034,7 @@ void FRawInputWindows::SendControllerEvents()
 			{
 				if (!DeviceAnalogData.KeyName.IsNone() && DeviceAnalogData.HasValue())
 				{
-					MessageHandler->OnControllerAnalog(DeviceAnalogData.KeyName, 0, DeviceAnalogData.GetValue());
+					MessageHandler->OnControllerAnalog(DeviceAnalogData.KeyName, UserId, DeviceId, DeviceAnalogData.GetValue());
 				}
 			}
 		}
