[RawInput] fixed crash when starting -game offscreen with RawInput plugin enabled

diff --git a/Engine/Plugins/Experimental/RawInput/Source/RawInput/Private/RawInput.cpp b/Engine/Plugins/Experimental/RawInput/Source/RawInput/Private/RawInput.cpp
--- a/Engine/Plugins/Experimental/RawInput/Source/RawInput/Private/RawInput.cpp
+++ b/Engine/Plugins/Experimental/RawInput/Source/RawInput/Private/RawInput.cpp
@@ -5,6 +5,8 @@
 #include "IInputDeviceModule.h"
 #include "IInputDevice.h"
 
+#include "Null/NullPlatformApplicationMisc.h"
+
 #if PLATFORM_WINDOWS
 	#include "Windows/RawInputWindows.h"
 #endif
@@ -265,6 +267,13 @@ IRawInput::IRawInput(const TSharedRef< FGenericApplicationMessageHandler >& InMe
 
 TSharedPtr< class IInputDevice > FRawInputPlugin::CreateInputDevice(const TSharedRef< FGenericApplicationMessageHandler >& InMessageHandler)
 {
+	// Don't create an input device when running offscreen
+	const bool bRunningOffscreen = FNullPlatformApplicationMisc::IsUsingNullApplication();
+	if (bRunningOffscreen)
+	{
+		return nullptr;
+	}
+
 	RawInputDevice = MakeShareable( new FPlatformRawInput(InMessageHandler) ); 
 	return RawInputDevice;
 }
